Fixed localization conflicts and hard-coded plurality

#rnx

diff --git a/Engine/Plugins/Enterprise/VariantManager/Source/VariantManager/Private/DisplayNodes/VariantManagerActorNode.cpp b/Engine/Plugins/Enterprise/VariantManager/Source/VariantManager/Private/DisplayNodes/VariantManagerActorNode.cpp
--- a/Engine/Plugins/Enterprise/VariantManager/Source/VariantManager/Private/DisplayNodes/VariantManagerActorNode.cpp
+++ b/Engine/Plugins/Enterprise/VariantManager/Source/VariantManager/Private/DisplayNodes/VariantManagerActorNode.cpp
@@ -147,9 +147,8 @@ TOptional<EItemDropZone> FVariantManagerActorNode::CanDrop(const FDragDropEvent&
 
 		if (NumActorsWeCanAdd > 0)
 		{
-			FText NewHoverText = FText::Format( LOCTEXT("CanDropActors", "Bind {0} actor{1} to variant '{2}'"),
-				FText::FromString(FString::FromInt(NumActorsWeCanAdd)),
-				NumActorsWeCanAdd != 1 ? FText::FromString(TEXT("s")) : FText(),
+			FText NewHoverText = FText::Format( LOCTEXT("CanDrop_BindActors", "Bind {0} {0}|plural(one=actor,other=actors) to variant '{1}'"),
+				NumActorsWeCanAdd,
 				Var->GetDisplayText());
 
 			const FSlateBrush* NewHoverIcon = FEditorStyle::GetBrush(TEXT("Graph.ConnectorFeedback.OK"));
@@ -160,7 +159,7 @@ TOptional<EItemDropZone> FVariantManagerActorNode::CanDrop(const FDragDropEvent&
 		}
 		else
 		{
-			FText NewHoverText = FText::Format( LOCTEXT("CanDropActors", "Actors already bound to variant '{0}'!"),
+			FText NewHoverText = FText::Format( LOCTEXT("CanDrop_ActorsAlreadyBound", "Actors already bound to variant '{0}'!"),
 				Var->GetDisplayText());
 
 			const FSlateBrush* NewHoverIcon = FEditorStyle::GetBrush(TEXT("Graph.ConnectorFeedback.Error"));
@@ -235,9 +234,8 @@ TOptional<EItemDropZone> FVariantManagerActorNode::CanDrop(const FDragDropEvent&
 		// Can copy new bindings
 		if (NumBindingsWeCanCopy > 0 && bIsCopy)
 		{
-			FText NewHoverText = FText::Format( LOCTEXT("CanDropActors", "Copy {0} actor binding{1} to variant '{2}'"),
-				FText::FromString(FString::FromInt(NumBindingsWeCanCopy)),
-				(NumBindingsWeCanCopy) != 1 ? FText::FromString(TEXT("s")) : FText(),
+			FText NewHoverText = FText::Format( LOCTEXT("CanDrop_CopyActorBindings", "Copy {0} actor {0}|plural(one=binding,other=bindings) to variant '{1}'"),
+				NumBindingsWeCanCopy,
 				Var->GetDisplayText());
 
 			const FSlateBrush* NewHoverIcon = FEditorStyle::GetBrush(TEXT("Graph.ConnectorFeedback.OK"));
@@ -249,9 +247,8 @@ TOptional<EItemDropZone> FVariantManagerActorNode::CanDrop(const FDragDropEvent&
 		// Have at least one binding we can move
 		else if ((NumOwnBindings+NumBindingsWeCanCopy) > 0 && !bIsCopy)
 		{
-			FText NewHoverText = FText::Format( LOCTEXT("CanDropActors", "Move {0} actor binding{1} to variant '{2}'"),
-				FText::FromString(FString::FromInt(NumOwnBindings + NumBindingsWeCanCopy)),
-				(NumOwnBindings + NumBindingsWeCanCopy) != 1 ? FText::FromString(TEXT("s")) : FText(),
+			FText NewHoverText = FText::Format( LOCTEXT("CanDrop_MoveActorBindings", "Move {0} actor {0}|plural(one=binding,other=bindings) to variant '{1}'"),
+				NumOwnBindings + NumBindingsWeCanCopy,
 				Var->GetDisplayText());
 
 			const FSlateBrush* NewHoverIcon = FEditorStyle::GetBrush(TEXT("Graph.ConnectorFeedback.OK"));
@@ -263,7 +260,7 @@ TOptional<EItemDropZone> FVariantManagerActorNode::CanDrop(const FDragDropEvent&
 		// If we at least were dragging some bindings (or else we were dragging variants/variant sets, etc)
 		else if(DraggedBindings.Num() > 0)
 		{
-			FText NewHoverText = FText::Format( LOCTEXT("CanDropActors", "All bindings already exist on variant '{0}'!"),
+			FText NewHoverText = FText::Format( LOCTEXT("CanDrop_AllBindingsAlreadyExist", "All bindings already exist on variant '{0}'!"),
 				Var->GetDisplayText());
 
 			const FSlateBrush* NewHoverIcon = FEditorStyle::GetBrush(TEXT("Graph.ConnectorFeedback.Error"));
@@ -324,8 +321,8 @@ void FVariantManagerActorNode::Drop(const FDragDropEvent& DragDropEvent, EItemDr
 			}
 
 			FScopedTransaction Transaction(FText::Format(
-				LOCTEXT("ActorNodeDropSceneActors", "Drop {0} scene actors near actor binding '{1}'"),
-				FText::AsNumber(Actors.Num()),
+				LOCTEXT("ActorNodeDropSceneActors", "Drop {0} scene {0}|plural(one=actor,other=actors) near actor binding '{1}'"),
+				Actors.Num(),
 				GetDisplayName()));
 
 			VarMan->CreateObjectBindingsAndCaptures(Actors, {Var}, TargetIndex);
@@ -390,8 +387,8 @@ void FVariantManagerActorNode::Drop(const FDragDropEvent& DragDropEvent, EItemDr
 			}
 
 			FScopedTransaction Transaction(FText::Format(
-				LOCTEXT("ActorNodeDropActors", "Drop {0} actor bindings near actor binding '{1}'"),
-				FText::AsNumber(BindingsWeCanDuplicate.Num()),
+				LOCTEXT("ActorNodeDropActors", "Drop {0} actor {0}|plural(one=binding,other=bindings) near actor binding '{1}'"),
+				BindingsWeCanDuplicate.Num(),
 				GetDisplayName()));
 
 			VarMan->DuplicateObjectBindings(BindingsWeCanDuplicate, Var, TargetIndex);
@@ -412,8 +409,8 @@ void FVariantManagerActorNode::Drop(const FDragDropEvent& DragDropEvent, EItemDr
 			}
 
 			FScopedTransaction Transaction(FText::Format(
-				LOCTEXT("ActorNodeDropActors", "Drop {0} actor bindings near actor binding '{1}'"),
-				FText::AsNumber(BindingsWeCanMove.Num()),
+				LOCTEXT("ActorNodeDropActors", "Drop {0} actor {0}|plural(one=binding,other=bindings) near actor binding '{1}'"),
+				BindingsWeCanMove.Num(),
 				GetDisplayName()));
 
 			VarMan->AddObjectBindings(BindingsWeCanMove, Var, TargetIndex, true);
diff --git a/Engine/Plugins/Enterprise/VariantManager/Source/VariantManager/Private/DisplayNodes/VariantManagerEnumPropertyNode.cpp b/Engine/Plugins/Enterprise/VariantManager/Source/VariantManager/Private/DisplayNodes/VariantManagerEnumPropertyNode.cpp
--- a/Engine/Plugins/Enterprise/VariantManager/Source/VariantManager/Private/DisplayNodes/VariantManagerEnumPropertyNode.cpp
+++ b/Engine/Plugins/Enterprise/VariantManager/Source/VariantManager/Private/DisplayNodes/VariantManagerEnumPropertyNode.cpp
@@ -185,7 +185,7 @@ FText FVariantManagerEnumPropertyNode::ComboboxGetText(bool bSameValue) const
 	{
 		if (!bSameValue)
 		{
-			return FText::FromString(TEXT("Multiple Values"));
+			return LOCTEXT("MultipleValuesLabel", "Multiple Values");
 		}
 
 		TSharedPtr<FString> SelectedText = Combobox->GetSelectedItem();
@@ -194,7 +194,7 @@ FText FVariantManagerEnumPropertyNode::ComboboxGetText(bool bSameValue) const
 			return FText::FromString(*SelectedText);
 		}
 
-		return FText::FromString(TEXT("(INVALID)"));
+		return LOCTEXT("InvalidLabel", "(INVALID)");
 	}
 
 	return FText();
diff --git a/Engine/Plugins/Enterprise/VariantManager/Source/VariantManager/Private/DisplayNodes/VariantManagerFunctionPropertyNode.cpp b/Engine/Plugins/Enterprise/VariantManager/Source/VariantManager/Private/DisplayNodes/VariantManagerFunctionPropertyNode.cpp
--- a/Engine/Plugins/Enterprise/VariantManager/Source/VariantManager/Private/DisplayNodes/VariantManagerFunctionPropertyNode.cpp
+++ b/Engine/Plugins/Enterprise/VariantManager/Source/VariantManager/Private/DisplayNodes/VariantManagerFunctionPropertyNode.cpp
@@ -44,12 +44,12 @@ void FVariantManagerFunctionPropertyNode::BuildContextMenu(FMenuBuilder& MenuBui
 
 FText FVariantManagerFunctionPropertyNode::GetDisplayName() const
 {
-	return FText::FromString(TEXT("Function caller"));
+	return LOCTEXT("NodeDisplayName", "Function caller");
 }
 
 FText FVariantManagerFunctionPropertyNode::GetDisplayNameToolTipText() const
 {
-	return FText::FromString(TEXT("Reference to one of the functions of the parent LevelVariantSets' FunctionDirector"));
+	return LOCTEXT("NodeToolTipText", "Reference to one of the functions of the parent LevelVariantSets' FunctionDirector");
 }
 
 TWeakObjectPtr<UVariantObjectBinding> FVariantManagerFunctionPropertyNode::GetObjectBinding() const
@@ -226,7 +226,7 @@ TSharedRef<SWidget> FVariantManagerFunctionPropertyNode::GetMenuContent()
 							else
 							{
 								TooltipText = FText::Format(
-									LOCTEXT("WrongPinTypeTooltip", "Input pin '{0}' of type '{1}' must be a reference\nto an object of a class parent to '{2}'"),
+									LOCTEXT("WrongPinTypeTooltipWithClass", "Input pin '{0}' of type '{1}' must be a reference\nto an object of a class parent to '{2}'"),
 									FText::FromName(PinInfo->PinName),
 									FText::FromString(PinClass->GetName()),
 									FText::FromString(BoundObjectClass->GetName()));
diff --git a/Engine/Plugins/Enterprise/VariantManager/Source/VariantManager/Private/DisplayNodes/VariantManagerPropertyNode.cpp b/Engine/Plugins/Enterprise/VariantManager/Source/VariantManager/Private/DisplayNodes/VariantManagerPropertyNode.cpp
--- a/Engine/Plugins/Enterprise/VariantManager/Source/VariantManager/Private/DisplayNodes/VariantManagerPropertyNode.cpp
+++ b/Engine/Plugins/Enterprise/VariantManager/Source/VariantManager/Private/DisplayNodes/VariantManagerPropertyNode.cpp
@@ -230,7 +230,7 @@ TSharedPtr<SWidget> FVariantManagerPropertyNode::GetPropertyValueWidget()
 		.HAlign(HAlign_Left)
 		[
 			SNew(STextBlock)
-			.Text(LOCTEXT("MultipleValuesText", "Multiple Values"))
+			.Text(LOCTEXT("MultipleValuesLabel", "Multiple Values"))
 			.Font(FEditorStyle::GetFontStyle("Sequencer.AnimationOutliner.RegularFont"))
 			.ColorAndOpacity(this, &FVariantManagerDisplayNode::GetDisplayNameColor)
 			.ToolTipText(LOCTEXT("MultipleValuesTooltip", "The selected actors have different values for this property"))
@@ -281,11 +281,11 @@ TSharedPtr<SWidget> FVariantManagerPropertyNode::GetPropertyValueWidget()
 		.HAlign(HAlign_Left)
 		[
 			SNew(STextBlock)
-			.Text(LOCTEXT("FailedToResolveText", "Unsupported property type!"))
+			.Text(LOCTEXT("UnsupportedPropertyType", "Unsupported property type!"))
 			.Font(FEditorStyle::GetFontStyle("Sequencer.AnimationOutliner.RegularFont"))
 			.ColorAndOpacity(this, &FVariantManagerDisplayNode::GetDisplayNameColor)
 			.ToolTipText(FText::Format(
-				LOCTEXT("FailedToResolveTooltip", "Properties of class '{0}' can't be captured yet!"),
+				LOCTEXT("UnsupportedPropertyTypeTooltip", "Properties of class '{0}' can't be captured yet!"),
 				FText::FromString(FirstPropertyValue->GetPropertyClass()->GetName())))
 		];
 	}
@@ -451,8 +451,8 @@ void FVariantManagerPropertyNode::UpdateRecordedDataFromSinglePropView(TSharedPt
 FReply FVariantManagerPropertyNode::ResetMultipleValuesToDefault()
 {
 	FScopedTransaction Transaction(FText::Format(
-		LOCTEXT("ResetPropertyValue", "Reset {0} property captures to default"),
-		FText::FromString(FString::FromInt(PropertyValues.Num()))
+		LOCTEXT("ResetPropertyValue", "Reset {0} property {0}|plural(one=capture,other=captures) to default"),
+		PropertyValues.Num()
 	));
 
 	TArray<uint8> DefaultValue;
@@ -486,8 +486,8 @@ FReply FVariantManagerPropertyNode::ResetMultipleValuesToDefault()
 FReply FVariantManagerPropertyNode::RecordMultipleValues()
 {
 	FScopedTransaction Transaction(FText::Format(
-		LOCTEXT("ResetPropertyValue", "Recorded new data for {0} property captures"),
-		FText::FromString(FString::FromInt(PropertyValues.Num()))
+		LOCTEXT("RecordedNewDataForProperty", "Recorded new data for {0} property {0}|plural(one=capture,other=captures)"),
+		PropertyValues.Num()
 	));
 
 	for (TWeakObjectPtr<UPropertyValue> PropertyValue : PropertyValues)
diff --git a/Engine/Plugins/Enterprise/VariantManager/Source/VariantManager/Private/DisplayNodes/VariantManagerStringPropertyNode.cpp b/Engine/Plugins/Enterprise/VariantManager/Source/VariantManager/Private/DisplayNodes/VariantManagerStringPropertyNode.cpp
--- a/Engine/Plugins/Enterprise/VariantManager/Source/VariantManager/Private/DisplayNodes/VariantManagerStringPropertyNode.cpp
+++ b/Engine/Plugins/Enterprise/VariantManager/Source/VariantManager/Private/DisplayNodes/VariantManagerStringPropertyNode.cpp
@@ -156,11 +156,11 @@ TSharedPtr<SWidget> FVariantManagerStringPropertyNode::GetPropertyValueWidget()
 		.HAlign(HAlign_Left)
 		[
 			SNew(STextBlock)
-			.Text(LOCTEXT("FailedToResolveText", "Unsupported property type!"))
+			.Text(LOCTEXT("UnsupportedPropertyType", "Unsupported property type!"))
 			.Font(FEditorStyle::GetFontStyle("Sequencer.AnimationOutliner.RegularFont"))
 			.ColorAndOpacity(this, &FVariantManagerDisplayNode::GetDisplayNameColor)
 			.ToolTipText(FText::Format(
-				LOCTEXT("FailedToResolveTooltip", "Properties of class '{0}' can't be captured yet!"),
+				LOCTEXT("UnsupportedPropertyTypeTooltip", "Properties of class '{0}' can't be captured yet!"),
 				FText::FromString(FirstPropertyValue->GetPropertyClass()->GetName())))
 		];
 	}
diff --git a/Engine/Plugins/Enterprise/VariantManager/Source/VariantManager/Private/DisplayNodes/VariantManagerStructPropertyNode.cpp b/Engine/Plugins/Enterprise/VariantManager/Source/VariantManager/Private/DisplayNodes/VariantManagerStructPropertyNode.cpp
--- a/Engine/Plugins/Enterprise/VariantManager/Source/VariantManager/Private/DisplayNodes/VariantManagerStructPropertyNode.cpp
+++ b/Engine/Plugins/Enterprise/VariantManager/Source/VariantManager/Private/DisplayNodes/VariantManagerStructPropertyNode.cpp
@@ -147,7 +147,7 @@ TSharedRef<SWidget> FVariantManagerStructPropertyNode::GenerateFloatEntryBox(UNu
 	.OnValueChanged(this, &FVariantManagerStructPropertyNode::OnFloatValueChanged, Prop)
 	.OnValueCommitted(this, &FVariantManagerStructPropertyNode::OnFloatPropCommitted, Prop, Offset)
 	.Value(this, &FVariantManagerStructPropertyNode::GetFloatValueFromCache, Prop)
-	.UndeterminedString(FText::FromString(TEXT("Multiple Values")))
+	.UndeterminedString(LOCTEXT("MultipleValuesLabel", "Multiple Values"))
 	.ShiftMouseMovePixelPerDelta(ShiftMouseMovePixelPerDelta)
 	.SupportDynamicSliderMaxValue(SupportDynamicSliderMaxValue)
 	.SupportDynamicSliderMinValue(SupportDynamicSliderMinValue)
@@ -182,7 +182,7 @@ TSharedRef<SWidget> FVariantManagerStructPropertyNode::GenerateSignedEntryBox(UN
 	.OnValueChanged(this, &FVariantManagerStructPropertyNode::OnSignedValueChanged, Prop)
 	.OnValueCommitted(this, &FVariantManagerStructPropertyNode::OnSignedPropCommitted, Prop, Offset)
 	.Value(this, &FVariantManagerStructPropertyNode::GetSignedValueFromCache, Prop)
-	.UndeterminedString(FText::FromString(TEXT("Multiple Values")))
+	.UndeterminedString(LOCTEXT("MultipleValuesLabel", "Multiple Values"))
 	.ShiftMouseMovePixelPerDelta(ShiftMouseMovePixelPerDelta)
 	.SupportDynamicSliderMaxValue(SupportDynamicSliderMaxValue)
 	.SupportDynamicSliderMinValue(SupportDynamicSliderMinValue)
@@ -217,7 +217,7 @@ TSharedRef<SWidget> FVariantManagerStructPropertyNode::GenerateUnsignedEntryBox(
 	.OnValueChanged(this, &FVariantManagerStructPropertyNode::OnUnsignedValueChanged, Prop)
 	.OnValueCommitted(this, &FVariantManagerStructPropertyNode::OnUnsignedPropCommitted, Prop, Offset)
 	.Value(this, &FVariantManagerStructPropertyNode::GetUnsignedValueFromCache, Prop)
-	.UndeterminedString(FText::FromString(TEXT("Multiple Values")))
+	.UndeterminedString(LOCTEXT("MultipleValuesLabel", "Multiple Values"))
 	.ShiftMouseMovePixelPerDelta(ShiftMouseMovePixelPerDelta)
 	.SupportDynamicSliderMaxValue(SupportDynamicSliderMaxValue)
 	.SupportDynamicSliderMinValue(SupportDynamicSliderMinValue)
@@ -512,7 +512,7 @@ void FVariantManagerStructPropertyNode::OnFloatPropCommitted(double InValue, ETe
 	}
 
 	FScopedTransaction Transaction(FText::Format(
-		LOCTEXT("StructPropertyNodeUpdateRecordedData", "Edit captured property '{0}'"),
+		LOCTEXT("EditCapturedProperty", "Edit captured property '{0}'"),
 		FText::FromName(FirstPropertyValue->GetPropertyName())));
 
 	TOptional<double>& StoredVal = FloatValues.FindOrAdd(Prop);
diff --git a/Engine/Plugins/Enterprise/VariantManager/Source/VariantManager/Private/DisplayNodes/VariantManagerVariantNode.cpp b/Engine/Plugins/Enterprise/VariantManager/Source/VariantManager/Private/DisplayNodes/VariantManagerVariantNode.cpp
--- a/Engine/Plugins/Enterprise/VariantManager/Source/VariantManager/Private/DisplayNodes/VariantManagerVariantNode.cpp
+++ b/Engine/Plugins/Enterprise/VariantManager/Source/VariantManager/Private/DisplayNodes/VariantManagerVariantNode.cpp
@@ -200,9 +200,8 @@ TOptional<EItemDropZone> FVariantManagerVariantNode::CanDrop(const FDragDropEven
 
 		if (NumActorsWeCanAdd > 0)
 		{
-			FText NewHoverText = FText::Format( LOCTEXT("CanDropActors", "Bind {0} actor{1} to variant '{2}'"),
-				FText::FromString(FString::FromInt(NumActorsWeCanAdd)),
-				NumActorsWeCanAdd != 1 ? FText::FromString(TEXT("s")) : FText(),
+			FText NewHoverText = FText::Format( LOCTEXT("CanDrop_BindActors", "Bind {0} {0}|plural(one=actor,other=actors) to variant '{2}'"),
+				NumActorsWeCanAdd,
 				GetVariant().GetDisplayText());
 
 			const FSlateBrush* NewHoverIcon = FEditorStyle::GetBrush(TEXT("Graph.ConnectorFeedback.OK"));
@@ -213,7 +212,7 @@ TOptional<EItemDropZone> FVariantManagerVariantNode::CanDrop(const FDragDropEven
 		}
 		else
 		{
-			FText NewHoverText = FText::Format( LOCTEXT("CanDropActors", "Actors already bound to variant '{0}'!"),
+			FText NewHoverText = FText::Format( LOCTEXT("CanDrop_ActorsAlreadyBound", "Actors already bound to variant '{0}'!"),
 				GetVariant().GetDisplayText());
 
 			const FSlateBrush* NewHoverIcon = FEditorStyle::GetBrush(TEXT("Graph.ConnectorFeedback.Error"));
@@ -272,10 +271,9 @@ TOptional<EItemDropZone> FVariantManagerVariantNode::CanDrop(const FDragDropEven
 			{
 				FModifierKeysState ModifierKeysState = FSlateApplication::Get().GetModifierKeys();
 
-				FText NewHoverText = FText::Format( LOCTEXT("CanDropActors", "{0} {1} actor binding{2} to variant '{3}'"),
-					ModifierKeysState.IsControlDown() ? FText::FromString(TEXT("Copy")) : FText::FromString(TEXT("Move")),
-					FText::FromString(FString::FromInt(NumActorsWeCanAdd)),
-					NumActorsWeCanAdd != 1 ? FText::FromString(TEXT("s")) : FText(),
+				FText NewHoverText = FText::Format( LOCTEXT("CanDrop_ApplyActors", "{0} {1} actor {1}|plural(one=binding,other=bindings) to variant '{2}'"),
+					ModifierKeysState.IsControlDown() ? LOCTEXT("Copy", "Copy") : LOCTEXT("Move", "Move"),
+					NumActorsWeCanAdd,
 					GetVariant().GetDisplayText());
 
 				const FSlateBrush* NewHoverIcon = FEditorStyle::GetBrush(TEXT("Graph.ConnectorFeedback.OK"));
@@ -286,7 +284,7 @@ TOptional<EItemDropZone> FVariantManagerVariantNode::CanDrop(const FDragDropEven
 			}
 			else
 			{
-				FText NewHoverText = FText::Format( LOCTEXT("CanDropActors", "Actors already bound to variant '{0}'!"),
+				FText NewHoverText = FText::Format( LOCTEXT("CanDrop_ActorsAlreadyBound", "Actors already bound to variant '{0}'!"),
 					GetVariant().GetDisplayText());
 
 				const FSlateBrush* NewHoverIcon = FEditorStyle::GetBrush(TEXT("Graph.ConnectorFeedback.Error"));
@@ -300,10 +298,9 @@ TOptional<EItemDropZone> FVariantManagerVariantNode::CanDrop(const FDragDropEven
 		{
 			FModifierKeysState ModifierKeysState = FSlateApplication::Get().GetModifierKeys();
 
-			FText NewHoverText = FText::Format( LOCTEXT("CanDropActors", "{0} {1} variant{2} to set '{3}'"),
-				ModifierKeysState.IsControlDown() ? FText::FromString(TEXT("Copy")) : FText::FromString(TEXT("Move")),
-				FText::FromString(FString::FromInt(DraggedVariants.Num())),
-				DraggedVariants.Num() != 1 ? FText::FromString(TEXT("s")) : FText(),
+			FText NewHoverText = FText::Format( LOCTEXT("CanDrop_ApplyVariants", "{0} {1} {1}|plural(one=variant,other=variants) to set '{2}'"),
+				ModifierKeysState.IsControlDown() ? LOCTEXT("Copy", "Copy") : LOCTEXT("Move", "Move"),
+				DraggedVariants.Num(),
 				ParentVarSet->GetDisplayText());
 
 			const FSlateBrush* NewHoverIcon = FEditorStyle::GetBrush(TEXT("Graph.ConnectorFeedback.OK"));
@@ -352,8 +349,8 @@ void FVariantManagerVariantNode::Drop(const FDragDropEvent& DragDropEvent, EItem
 		}
 
 		FScopedTransaction Transaction(FText::Format(
-			LOCTEXT("VariantNodeDropSceneActors", "Drop {0} scene actors on variant '{1}'"),
-			FText::AsNumber(Actors.Num()),
+			LOCTEXT("VariantNodeDropSceneActors", "Drop {0} scene {0}|plural(one=actor,other=actors) on variant '{1}'"),
+			Actors.Num(),
 			GetDisplayName()));
 
 		VarMan->CreateObjectBindingsAndCaptures(Actors, {Var});
@@ -401,8 +398,8 @@ void FVariantManagerVariantNode::Drop(const FDragDropEvent& DragDropEvent, EItem
 		if (DraggedBindings.Num() > 0)
 		{
 			FScopedTransaction Transaction(FText::Format(
-				LOCTEXT("VariantNodeDropBindings", "Drop {0} actor bindings on variant '{1}'"),
-				FText::AsNumber(DraggedBindings.Num()),
+				LOCTEXT("VariantNodeDropBindings", "Drop {0} actor {0}|plural(one=binding,other=bindings) on variant '{1}'"),
+				DraggedBindings.Num(),
 				GetDisplayName()));
 
 			FModifierKeysState ModifierKeysState = FSlateApplication::Get().GetModifierKeys();
@@ -427,8 +424,8 @@ void FVariantManagerVariantNode::Drop(const FDragDropEvent& DragDropEvent, EItem
 		else if (DraggedVariants.Num() > 0)
 		{
 			FScopedTransaction Transaction(FText::Format(
-				LOCTEXT("VariantNodeDropVariants", "Drop {0} variants near variant '{1}'"),
-				FText::AsNumber(DraggedVariants.Num()),
+				LOCTEXT("VariantNodeDropVariants", "Drop {0} {0}|plural(one=variant,other=variants) near variant '{1}'"),
+				DraggedVariants.Num(),
 				GetDisplayName()));
 
 			UVariantSet* ParentVarSet = Var->GetParent();
diff --git a/Engine/Plugins/Enterprise/VariantManager/Source/VariantManager/Private/DisplayNodes/VariantManagerVariantSetNode.cpp b/Engine/Plugins/Enterprise/VariantManager/Source/VariantManager/Private/DisplayNodes/VariantManagerVariantSetNode.cpp
--- a/Engine/Plugins/Enterprise/VariantManager/Source/VariantManager/Private/DisplayNodes/VariantManagerVariantSetNode.cpp
+++ b/Engine/Plugins/Enterprise/VariantManager/Source/VariantManager/Private/DisplayNodes/VariantManagerVariantSetNode.cpp
@@ -252,10 +252,9 @@ TOptional<EItemDropZone> FVariantManagerVariantSetNode::CanDrop(const FDragDropE
 
 		if (NumActorsWeCanAdd > 0 && NumVarsThatCanAccept > 0)
 		{
-			FText NewHoverText = FText::Format( LOCTEXT("CanDropActors", "Bind {0} actor{1} to {2} variants"),
-				FText::FromString(FString::FromInt(NumActorsWeCanAdd)),
-				NumActorsWeCanAdd != 1 ? FText::FromString(TEXT("s")) : FText(),
-				FText::FromString(FString::FromInt(NumVarsThatCanAccept)));
+			FText NewHoverText = FText::Format( LOCTEXT("CanDrop_BindActors", "Bind {0} {0}|plural(one=actor,other=actors) to {1} {1}|plural(one=variant,other=variants)"),
+				NumActorsWeCanAdd,
+				NumVarsThatCanAccept);
 
 			const FSlateBrush* NewHoverIcon = FEditorStyle::GetBrush(TEXT("Graph.ConnectorFeedback.OK"));
 
@@ -265,7 +264,7 @@ TOptional<EItemDropZone> FVariantManagerVariantSetNode::CanDrop(const FDragDropE
 		}
 		else if (GetChildNodes().Num() < 1)
 		{
-			FText NewHoverText = FText::Format( LOCTEXT("CanDropActors", "Variant set '{0}' has no variants!"),
+			FText NewHoverText = FText::Format( LOCTEXT("CanDrop_NoVariantsInSet", "Variant set '{0}' has no variants!"),
 				GetVariantSet().GetDisplayText());
 
 			const FSlateBrush* NewHoverIcon = FEditorStyle::GetBrush(TEXT("Graph.ConnectorFeedback.Error"));
@@ -276,7 +275,7 @@ TOptional<EItemDropZone> FVariantManagerVariantSetNode::CanDrop(const FDragDropE
 		}
 		else
 		{
-			FText NewHoverText = FText::Format( LOCTEXT("CanDropActors", "Actors already bound to all variants of set '{0}'!"),
+			FText NewHoverText = FText::Format( LOCTEXT("CanDrop_ActorsAlreadyBound", "Actors already bound to all variants of set '{0}'!"),
 				GetVariantSet().GetDisplayText());
 
 			const FSlateBrush* NewHoverIcon = FEditorStyle::GetBrush(TEXT("Graph.ConnectorFeedback.Error"));
@@ -360,11 +359,9 @@ TOptional<EItemDropZone> FVariantManagerVariantSetNode::CanDrop(const FDragDropE
 			}
 			else if (NumActorsWeCanAdd > 0)
 			{
-				FText NewHoverText = FText::Format( LOCTEXT("CanDropActors", "Copy {0} actor binding{1} to {2} variant{3}"),
-					FText::FromString(FString::FromInt(NumActorsWeCanAdd)),
-					NumActorsWeCanAdd != 1 ? FText::FromString(TEXT("s")) : FText(),
-					FText::FromString(FString::FromInt(NumVariants)),
-					NumVariants != 1 ? FText::FromString(TEXT("s")) : FText());
+				FText NewHoverText = FText::Format( LOCTEXT("CanDrop_CopyActors", "Copy {0} actor {0}|plural(one=binding,other=bindings) to {1} {1}|plural(one=variant,other=variants)"),
+					NumActorsWeCanAdd,
+					NumVariants);
 
 				const FSlateBrush* NewHoverIcon = FEditorStyle::GetBrush(TEXT("Graph.ConnectorFeedback.OK"));
 
@@ -374,7 +371,7 @@ TOptional<EItemDropZone> FVariantManagerVariantSetNode::CanDrop(const FDragDropE
 			}
 			else
 			{
-				FText NewHoverText = FText::Format( LOCTEXT("CanDropActors", "Actors already bound to all variants of set '{0}'!"),
+				FText NewHoverText = FText::Format( LOCTEXT("CanDrop_ActorsAlreadyBound", "Actors already bound to all variants of set '{0}'!"),
 					GetVariantSet().GetDisplayText());
 
 				const FSlateBrush* NewHoverIcon = FEditorStyle::GetBrush(TEXT("Graph.ConnectorFeedback.Error"));
@@ -388,10 +385,9 @@ TOptional<EItemDropZone> FVariantManagerVariantSetNode::CanDrop(const FDragDropE
 		{
 			FModifierKeysState ModifierKeysState = FSlateApplication::Get().GetModifierKeys();
 
-			FText NewHoverText = FText::Format( LOCTEXT("CanDropActors", "{0} {1} variant{2} to set '{3}'"),
-				ModifierKeysState.IsControlDown() ? FText::FromString(TEXT("Copy")) : FText::FromString(TEXT("Move")),
-				FText::FromString(FString::FromInt(NumDraggedVariants)),
-				NumDraggedVariants != 1 ? FText::FromString(TEXT("s")) : FText(),
+			FText NewHoverText = FText::Format( LOCTEXT("CanDrop_ApplyVariantsToSet", "{0} {1} {1}|plural(one=variant,other=variants) to set '{3}'"),
+				ModifierKeysState.IsControlDown() ? LOCTEXT("Copy", "Copy") : LOCTEXT("Move", "Move"),
+				NumDraggedVariants,
 				GetVariantSet().GetDisplayText());
 
 			const FSlateBrush* NewHoverIcon = FEditorStyle::GetBrush(TEXT("Graph.ConnectorFeedback.OK"));
@@ -404,10 +400,9 @@ TOptional<EItemDropZone> FVariantManagerVariantSetNode::CanDrop(const FDragDropE
 		{
 			FModifierKeysState ModifierKeysState = FSlateApplication::Get().GetModifierKeys();
 
-			FText NewHoverText = FText::Format( LOCTEXT("CanDropActors", "{0} {1} variant set{2}"),
-				ModifierKeysState.IsControlDown() ? FText::FromString(TEXT("Copy")) : FText::FromString(TEXT("Move")),
-				FText::FromString(FString::FromInt(NumDraggedVariants)),
-				NumDraggedVariants != 1 ? FText::FromString(TEXT("s")) : FText());
+			FText NewHoverText = FText::Format( LOCTEXT("CanDrop_ApplyVariantToSets", "{0} {1} variant {1}|plural(one=set,other=sets)"),
+				ModifierKeysState.IsControlDown() ? LOCTEXT("Copy", "Copy") : LOCTEXT("Move", "Move"),
+				NumDraggedVariants);
 
 			const FSlateBrush* NewHoverIcon = FEditorStyle::GetBrush(TEXT("Graph.ConnectorFeedback.OK"));
 
@@ -470,8 +465,8 @@ void FVariantManagerVariantSetNode::Drop(const FDragDropEvent& DragDropEvent, EI
 		}
 
 		FScopedTransaction Transaction(FText::Format(
-			LOCTEXT("VariantSetNodeDropSceneActors", "Drop {0} scene actors on variant set '{1}'"),
-			FText::AsNumber(Actors.Num()),
+			LOCTEXT("VariantSetNodeDropSceneActors", "Drop {0} scene {0}|plural(one=actor,other=actors) on variant set '{1}'"),
+			Actors.Num(),
 			GetDisplayName()));
 
 		// Add bindings just once, or else it will spawn multiple popups and trigger multiple refreshes
@@ -529,8 +524,8 @@ void FVariantManagerVariantSetNode::Drop(const FDragDropEvent& DragDropEvent, EI
 		if (NumDraggedBindings > 0)
 		{
 			FScopedTransaction Transaction(FText::Format(
-				LOCTEXT("VariantSetNodeDropBindings", "Drop {0} actor bindings on variant set '{1}'"),
-				FText::AsNumber(NumDraggedBindings),
+				LOCTEXT("VariantSetNodeDropBindings", "Drop {0} actor {0}|plural(one=binding,other=bindings) on variant set '{1}'"),
+				NumDraggedBindings,
 				GetDisplayName()));
 
 			// Copy our bindings to all our children
@@ -565,8 +560,8 @@ void FVariantManagerVariantSetNode::Drop(const FDragDropEvent& DragDropEvent, EI
 		else if (NumDraggedVariants > 0)
 		{
 			FScopedTransaction Transaction(FText::Format(
-				LOCTEXT("VariantSetNodeDropVariants", "Drop {0} variants on variant set '{1}'"),
-				FText::AsNumber(NumDraggedVariants),
+				LOCTEXT("VariantSetNodeDropVariants", "Drop {0} {0}|plural(one=variant,other=variants) on variant set '{1}'"),
+				NumDraggedVariants,
 				GetDisplayName()));
 
 			FModifierKeysState ModifierKeysState = FSlateApplication::Get().GetModifierKeys();
@@ -591,8 +586,8 @@ void FVariantManagerVariantSetNode::Drop(const FDragDropEvent& DragDropEvent, EI
 		else if (NumDraggedVariantSets > 0)
 		{
 			FScopedTransaction Transaction(FText::Format(
-				LOCTEXT("VariantSetNodeDropVariantSets", "Drop {0} variant sets near variant set '{1}'"),
-				FText::AsNumber(NumDraggedVariantSets),
+				LOCTEXT("VariantSetNodeDropVariantSets", "Drop {0} variant {0}|plural(one=set,other=sets) near variant set '{1}'"),
+				NumDraggedVariantSets,
 				GetDisplayName()));
 
 			ULevelVariantSets* ParentLevelVarSet = VarSet->GetParent();
diff --git a/Engine/Plugins/Enterprise/VariantManager/Source/VariantManager/Private/SVariantManager.cpp b/Engine/Plugins/Enterprise/VariantManager/Source/VariantManager/Private/SVariantManager.cpp
--- a/Engine/Plugins/Enterprise/VariantManager/Source/VariantManager/Private/SVariantManager.cpp
+++ b/Engine/Plugins/Enterprise/VariantManager/Source/VariantManager/Private/SVariantManager.cpp
@@ -606,9 +606,9 @@ void SVariantManager::AddEditorSelectedActorsToVariant()
 	}
 
 	FScopedTransaction Transaction(FText::Format(
-		LOCTEXT("AddEditorSelectedActorsToVariantTransaction", "Add {0} actor bindings to {1} variants"),
-		FText::FromString(FString::FromInt(Actors.Num())),
-		FText::FromString(FString::FromInt(SelectedVariants.Num()))
+		LOCTEXT("AddEditorSelectedActorsToVariantTransaction", "Add {0} actor {0}|plural(one=binding,other=bindings) to {1} {1}|plural(one=variant,other=variants)"),
+		Actors.Num(),
+		SelectedVariants.Num()
 	));
 
 	VariantManager->CreateObjectBindingsAndCaptures(Actors, SelectedVariants);
@@ -733,9 +733,9 @@ void SVariantManager::CutSelectionVariantTree()
 	// Don't capture CopySelection in the transaction buffer because if we undo we kind of expect
 	// our cut stuff to still be in the clipboard
 	FScopedTransaction Transaction(FText::Format(
-		LOCTEXT("CutSelectionVariantTreeTransaction", "Cut {0} variants and {1} variant sets"),
-		FText::FromString(FString::FromInt(CopiedVariants.Num())),
-		FText::FromString(FString::FromInt(CopiedVariantSets.Num()))
+		LOCTEXT("CutSelectionVariantTreeTransaction", "Cut {0} {0}|plural(one=variant,other=variants) and {1} variant {1}|plural(one=set,other=sets)"),
+		CopiedVariants.Num(),
+		CopiedVariantSets.Num()
 		));
 
 
@@ -782,9 +782,9 @@ void SVariantManager::PasteSelectionVariantTree()
 	const TArray<TStrongObjectPtr<UVariant>>& CopiedVariants = FVariantManagerClipboard::GetVariants();
 
 	FScopedTransaction Transaction(FText::Format(
-		LOCTEXT("PasteSelectionVariantTreeTransaction", "Paste {0} variants and {1} variant sets"),
-		FText::FromString(FString::FromInt(CopiedVariants.Num())),
-		FText::FromString(FString::FromInt(CopiedVariantSets.Num()))
+		LOCTEXT("PasteSelectionVariantTreeTransaction", "Paste {0} {0}|plural(one=variant,other=variants) and {1} variant {1}|plural(one=set,other=sets)"),
+		CopiedVariants.Num(),
+		CopiedVariantSets.Num()
 		));
 
 	// Paste variant sets onto the tree, regardless of where we clicked
@@ -888,9 +888,9 @@ void SVariantManager::DeleteSelectionVariantTree()
 	VariantManager->GetSelection().GetSelectedVariantsAndVariantSets(VariantsToDelete, VariantSetsToDelete);
 
 	FScopedTransaction Transaction(FText::Format(
-		LOCTEXT("DeleteSelectionVariantTreeTransaction", "Delete {0} variants and {1} variant sets"),
-		FText::FromString(FString::FromInt(VariantsToDelete.Num())),
-		FText::FromString(FString::FromInt(VariantSetsToDelete.Num()))
+		LOCTEXT("DeleteSelectionVariantTreeTransaction", "Delete {0} {0}|plural(one=variant,other=variants) and {1} variant {1}|plural(one=set,other=sets)"),
+		VariantsToDelete.Num(),
+		VariantSetsToDelete.Num()
 		));
 
 	VariantManager->RemoveVariantsFromParent(VariantsToDelete);
@@ -922,9 +922,9 @@ void SVariantManager::DuplicateSelectionVariantTree()
 	TArray<UVariant*> NewVariants;
 
 	FScopedTransaction Transaction(FText::Format(
-		LOCTEXT("DuplicateSelectionVariantTreeTransaction", "Duplicate {0} variants and {1} variant sets"),
-		FText::FromString(FString::FromInt(VariantsToDuplicate.Num())),
-		FText::FromString(FString::FromInt(VariantSetsToDuplicate.Num()))
+		LOCTEXT("DuplicateSelectionVariantTreeTransaction", "Duplicate {0} {0}|plural(one=variant,other=variants) and {1} variant {1}|plural(one=set,other=sets)"),
+		VariantsToDuplicate.Num(),
+		VariantSetsToDuplicate.Num()
 		));
 
 	// Duplicate variants
@@ -1237,9 +1237,8 @@ void SVariantManager::CaptureNewPropertiesFromSelectedActors()
 		int32 NumBindings = SelectedBindings.Num();
 
 		FScopedTransaction Transaction(FText::Format(
-			LOCTEXT("ActorNodeCaptureNewPropertiesTransaction", "Capture new properties for {0} actor binding{1}"),
-			FText::FromString(FString::FromInt(NumBindings)),
-			NumBindings != 1 ? FText::FromString(TEXT("s")) : FText()
+			LOCTEXT("ActorNodeCaptureNewPropertiesTransaction", "Capture new properties for {0} actor {0}|plural(one=binding,other=bindings)"),
+			NumBindings
 			));
 
 		PinnedVariantManager->CaptureNewProperties(SelectedBindings);
@@ -1278,9 +1277,8 @@ void SVariantManager::AddFunctionCaller()
 
 	int32 NumNewCallers = SelectedActorNodes.Num();
 	FScopedTransaction Transaction(FText::Format(
-		LOCTEXT("AddFunctionCaller", "Created {0} new function caller{1}"),
-		FText::FromString(FString::FromInt(NumNewCallers)),
-		NumNewCallers != 1 ? FText::FromString(TEXT("s")) : FText()
+		LOCTEXT("AddFunctionCaller", "Created {0} new function {0}|plural(one=caller,other=callers)"),
+		NumNewCallers
 	));
 
 	for (const TSharedPtr<FVariantManagerActorNode>& Node : SelectedActorNodes)
@@ -1330,9 +1328,8 @@ void SVariantManager::RemoveActorBindings()
 		int32 NumBindings = SelectedBindings.Num();
 
 		FScopedTransaction Transaction(FText::Format(
-			LOCTEXT("ActorNodeRemoveTransaction", "Remove {0} actor binding{1}"),
-			FText::FromString(FString::FromInt(NumBindings)),
-			NumBindings != 1 ? FText::FromString(TEXT("s")) : FText()
+			LOCTEXT("ActorNodeRemoveTransaction", "Remove {0} actor {0}|plural(one=binding,other=bindings)"),
+			NumBindings
 		));
 
 		PinnedVariantManager->RemoveObjectBindingsFromParent(SelectedBindings);
@@ -1458,8 +1455,8 @@ void SVariantManager::RemoveCapture()
 	}
 
 	FScopedTransaction Transaction(FText::Format(
-		LOCTEXT("RemoveCaptureTransaction", "Remove {0} property captures"),
-		FText::AsNumber(PropValuesToRemove.Num())));
+		LOCTEXT("RemoveCaptureTransaction", "Remove {0} property {0}|plural(one=capture,other=captures)"),
+		PropValuesToRemove.Num()));
 
 	PinnedVariantManager->RemovePropertyCapturesFromParent(PropValuesToRemove);
 
@@ -1528,8 +1525,8 @@ void SVariantManager::RemoveDirectorFunctionCaller()
 	}
 
 	FScopedTransaction Transaction(FText::Format(
-		LOCTEXT("RemoveCaptureTransaction", "Remove {0} function callers"),
-		FText::AsNumber(NumCallersWellRemove)));
+		LOCTEXT("RemoveCallersTransaction", "Remove {0} function {0}|plural(one=caller,other=callers)"),
+		NumCallersWellRemove));
 
 	for (auto& Pair : FunctionCallers)
 	{
diff --git a/Engine/Plugins/Enterprise/VariantManager/Source/VariantManager/Private/SVariantManagerActorListView.cpp b/Engine/Plugins/Enterprise/VariantManager/Source/VariantManager/Private/SVariantManagerActorListView.cpp
--- a/Engine/Plugins/Enterprise/VariantManager/Source/VariantManager/Private/SVariantManagerActorListView.cpp
+++ b/Engine/Plugins/Enterprise/VariantManager/Source/VariantManager/Private/SVariantManagerActorListView.cpp
@@ -92,11 +92,9 @@ FReply SVariantManagerActorListView::OnDragOver( const FGeometry& MyGeometry, co
 
 		if (NumActors > 0)
 		{
-			FText NewHoverText = FText::Format( LOCTEXT("CanDropActors", "Bind {0} actor{1} to {2} selected variant{3}"),
-				FText::FromString(FString::FromInt(NumActors)),
-				NumActors != 1 ? FText::FromString(TEXT("s")) : FText(),
-				FText::FromString(FString::FromInt(VariantsThatCanAccept)),
-				VariantsThatCanAccept != 1 ? FText::FromString(TEXT("s")) : FText());
+			FText NewHoverText = FText::Format( LOCTEXT("CanDropActors", "Bind {0} {0}|plural(one=actor,other=actors) to {1} selected {1}|plural(one=variant,other=variants)"),
+				NumActors,
+				VariantsThatCanAccept);
 
 			const FSlateBrush* NewHoverIcon = FEditorStyle::GetBrush(TEXT("Graph.ConnectorFeedback.OK"));
 
@@ -163,9 +161,9 @@ FReply SVariantManagerActorListView::OnDrop( const FGeometry& MyGeometry, const
 			}
 
 			FScopedTransaction Transaction(FText::Format(
-				LOCTEXT("ActorListViewDropVariants", "Drop {0} scene actors on {1} variants"),
-				FText::FromString(FString::FromInt(DraggedRawActors.Num())),
-				FText::FromString(FString::FromInt(SelectedVariants.Num()))));
+				LOCTEXT("ActorListViewDropVariants", "Drop {0} scene {0}|plural(one=actor,other=actors) on {1} {1}|plural(one=variant,other=variants)"),
+				DraggedRawActors.Num(),
+				SelectedVariants.Num()));
 
 			VarMan->CreateObjectBindingsAndCaptures(DraggedRawActors, SelectedVariants);
 
diff --git a/Engine/Plugins/Enterprise/VariantManager/Source/VariantManager/Private/SVariantManagerTableRow.cpp b/Engine/Plugins/Enterprise/VariantManager/Source/VariantManager/Private/SVariantManagerTableRow.cpp
--- a/Engine/Plugins/Enterprise/VariantManager/Source/VariantManager/Private/SVariantManagerTableRow.cpp
+++ b/Engine/Plugins/Enterprise/VariantManager/Source/VariantManager/Private/SVariantManagerTableRow.cpp
@@ -60,9 +60,8 @@ FReply SVariantManagerTableRow::DragDetected( const FGeometry& InGeometry, const
 
 			int32 NumNodes = DraggableNodes.Num();
 
-			DefaultHoverText = FText::Format(NSLOCTEXT("VariantManagerTableRow", "DefaultDragText", "{0} actor node{1}" ),
-				FText::AsNumber(NumNodes),
-				NumNodes != 1 ? FText::FromString(TEXT("s")) : FText());
+			DefaultHoverText = FText::Format(NSLOCTEXT("VariantManagerTableRow", "DragActorNode", "{0} actor {0}|plural(one=node,other=nodes)" ),
+				NumNodes);
 
 			break;
 		}
@@ -81,8 +80,8 @@ FReply SVariantManagerTableRow::DragDetected( const FGeometry& InGeometry, const
 
 			int32 NumNodes = DraggableNodes.Num();
 
-			DefaultHoverText = FText::Format(NSLOCTEXT("VariantManagerTableRow", "DefaultDragText", "{0} variants and variant sets" ),
-				FText::AsNumber(NumNodes));
+			DefaultHoverText = FText::Format(NSLOCTEXT("VariantManagerTableRow", "DragVariants", "{0} {0}|plural(one=variant,other=variants) and/or variant {0}|plural(one=set,other=sets)" ),
+				NumNodes);
 
 			break;
 		}
