Copied cl 9451133, 9451598, 9451650, 9451750, 9451891, 9451927, 9451999, 9452174

C4D: Work around allocation bug in Melange SDK
C4D: Fix bug when removing intermediary null actors with transforms
C4D: Expose to Blueprint/Python just import options that are visible in the GUI

VariantManager: Fix light/fog color captures not detecting when active/inactive
VariantManager: Fix bugs with auto-expose properties
- Fix materials not being auto-captured;
- Fix additional actors being captured when dragging actors to the level;
- Fix PPV properties not updating after captured;
- Make cached properties resilient against hot reload;
VariantManager: Manual Mobility control for SwitchActor

VRED/DG: Defense against failing to parse variants or animations

diff --git a/Engine/Plugins/Enterprise/DatasmithC4DImporter/Source/DatasmithC4DTranslator/Private/DatasmithC4DImportOptions.h b/Engine/Plugins/Enterprise/DatasmithC4DImporter/Source/DatasmithC4DTranslator/Private/DatasmithC4DImportOptions.h
--- a/Engine/Plugins/Enterprise/DatasmithC4DImporter/Source/DatasmithC4DTranslator/Private/DatasmithC4DImportOptions.h
+++ b/Engine/Plugins/Enterprise/DatasmithC4DImporter/Source/DatasmithC4DTranslator/Private/DatasmithC4DImportOptions.h
@@ -29,20 +29,20 @@ public:
 	/**
 	 * Import Mesh With No Vertex, false by default
 	 */
-	UPROPERTY(config, EditAnywhere, BlueprintReadWrite, Category = DebugProperty, meta = (DisplayName = "Import Mesh With No Vertex"))
+	UPROPERTY(config, EditAnywhere, Category = DebugProperty, meta = (DisplayName = "Import Mesh With No Vertex"))
 	bool bImportEmptyMesh;
 
 	/**
 	 * Remove empty actors that have only one child, false by default
 	 */
-	UPROPERTY(config, EditAnywhere, BlueprintReadWrite, Category = DebugProperty, meta = (DisplayName = "Clean the Scene of Empty Actors With Only One Child"))
+	UPROPERTY(config, EditAnywhere, Category = DebugProperty, meta = (DisplayName = "Clean the Scene of Empty Actors With Only One Child"))
 	bool bOptimizeEmptySingleChildActors;
 
 	/**
 	* Ignore the normals provided by Melange and let Datasmith generate them
 	* The "if (RawNormal.SizeSquared() < SMALL_NUMBER)" check in see FDatasmithMeshUtils::ToMeshDescription is preventing some meshes to load
 	*/
-	UPROPERTY(config, EditAnywhere, BlueprintReadWrite, Category = DebugProperty, meta = (DisplayName = "Generate the Normals"))
+	UPROPERTY(config, EditAnywhere, Category = DebugProperty, meta = (DisplayName = "Generate the Normals"))
 	bool bAlwaysGenerateNormals;
 
 	/**
@@ -50,13 +50,13 @@ public:
 	* It is needed when scene has too precise meshes because the conversion from double (C4D) to float (UE4)
 	* can "degenerate" the faces
 	*/
-	UPROPERTY(config, EditAnywhere, BlueprintReadWrite, Category = DebugProperty, meta = (DisplayName = "Scale the Entire Scene"))
+	UPROPERTY(config, EditAnywhere, Category = DebugProperty, meta = (DisplayName = "Scale the Entire Scene"))
 	float ScaleVertices;
 
 	/**
 	* Export the imported scene as a .datasmith file, next to the .c4d file.
 	*/
-	UPROPERTY(config, EditAnywhere, BlueprintReadWrite, Category = DebugProperty, meta = (DisplayName = "Export to .udatasmith"))
+	UPROPERTY(config, EditAnywhere, Category = DebugProperty, meta = (DisplayName = "Export to .udatasmith"))
 	bool bExportToUDatasmith;
 };
 
diff --git a/Engine/Plugins/Enterprise/DatasmithC4DImporter/Source/DatasmithC4DTranslator/Private/DatasmithC4DImporter.cpp b/Engine/Plugins/Enterprise/DatasmithC4DImporter/Source/DatasmithC4DTranslator/Private/DatasmithC4DImporter.cpp
--- a/Engine/Plugins/Enterprise/DatasmithC4DImporter/Source/DatasmithC4DTranslator/Private/DatasmithC4DImporter.cpp
+++ b/Engine/Plugins/Enterprise/DatasmithC4DImporter/Source/DatasmithC4DTranslator/Private/DatasmithC4DImporter.cpp
@@ -499,8 +499,8 @@ void RemoveEmptyActorsRecursive(TSharedPtr<IDatasmithActorElement>& Actor, const
 			{
 				TSharedPtr<IDatasmithActorElement> GrandChild = Child->GetChild(GrandChildIndex);
 
-				Actor->AddChild(GrandChild);
 				Child->RemoveChild(GrandChild);
+				Actor->AddChild(GrandChild);
 			}
 
 			Actor->RemoveChild(Child);
@@ -2746,9 +2746,20 @@ bool FDatasmithC4DImporter::OpenFile(const FString& InFilename)
 
 	if (C4Dfile->Open(DOC_IDENT, TCHAR_TO_ANSI(*InFilename), FILEOPEN_READ))
 	{
-		if (!C4dDocument->ReadObject(C4Dfile, true))
+		bool bSuccess = C4dDocument->ReadObject(C4Dfile, true);
+
+		int64 LastPos = static_cast<int64>(C4Dfile->GetPosition());
+		int64 Length = static_cast<int64>(C4Dfile->GetLength());
+		int64 Version = static_cast<int64>(C4Dfile->GetFileVersion());
+		FILEERROR Error = C4Dfile->GetError();
+
+		if (bSuccess)
+		{
+			UE_LOG(LogDatasmithC4DImport, Log, TEXT("Melange SDK successfully read the file '%s' (read %ld out of %ld bytes, version %ld)"), *InFilename, LastPos, Length, Version);
+		}
+		else
 		{
-			UE_LOG(LogDatasmithC4DImport, Warning, TEXT("Melange did not successfully read the entire file '%s'. Imported scene may contain errors or missing data."), *InFilename);
+			UE_LOG(LogDatasmithC4DImport, Warning, TEXT("Melange SDK did not read the entire file '%s' (read %ld out of %ld bytes, version %ld, error code: %d). Imported scene may contain errors or missing data."), *InFilename, LastPos, Length, Version, Error);
 		}
 	}
 	else
diff --git a/Engine/Plugins/Enterprise/DatasmithC4DImporter/Source/DatasmithC4DTranslator/Private/DatasmithC4DMelangeSDKOverload.cpp b/Engine/Plugins/Enterprise/DatasmithC4DImporter/Source/DatasmithC4DTranslator/Private/DatasmithC4DMelangeSDKOverload.cpp
--- a/Engine/Plugins/Enterprise/DatasmithC4DImporter/Source/DatasmithC4DTranslator/Private/DatasmithC4DMelangeSDKOverload.cpp
+++ b/Engine/Plugins/Enterprise/DatasmithC4DImporter/Source/DatasmithC4DTranslator/Private/DatasmithC4DMelangeSDKOverload.cpp
@@ -20,7 +20,10 @@ namespace melange
 
 	void* MemAllocNC(Int size)
 	{
-		return FMemory::Malloc(size);
+		// Patch for Melange SDK version 20.004_RBMelange20.0_259890
+		// It seems to rely on malloc(0) being non-nullptr, which is not true for FMemory::Malloc.
+		// Without this, it will fail to completely read files with geometry that have deleted faces
+		return FMemory::Malloc(size > 0 ? size : 1);
 	}
 
 	void* MemAlloc(Int size)
diff --git a/Engine/Plugins/Enterprise/DatasmithFBXImporter/Source/DatasmithDeltaGenTranslator/Private/DatasmithDeltaGenImportOptions.h b/Engine/Plugins/Enterprise/DatasmithFBXImporter/Source/DatasmithDeltaGenTranslator/Private/DatasmithDeltaGenImportOptions.h
--- a/Engine/Plugins/Enterprise/DatasmithFBXImporter/Source/DatasmithDeltaGenTranslator/Private/DatasmithDeltaGenImportOptions.h
+++ b/Engine/Plugins/Enterprise/DatasmithFBXImporter/Source/DatasmithDeltaGenTranslator/Private/DatasmithDeltaGenImportOptions.h
@@ -28,37 +28,37 @@ class UDatasmithDeltaGenImportOptions : public UDatasmithFBXImportOptions
 	GENERATED_UCLASS_BODY()
 
 public:
-	UPROPERTY(config, EditAnywhere, Category=Processing, meta=(ToolTip="Merge nodes that are not used in variants/animations"))
+	UPROPERTY(config, EditAnywhere, BlueprintReadWrite, Category=Processing, meta=(ToolTip="Merge nodes that are not used in variants/animations"))
 	bool bMergeNodes;
 
-	UPROPERTY(config, EditAnywhere, Category = Processing, meta = (ToolTip = "Finds duplicated scene nodes which could be replaced with a single instance"))
+	UPROPERTY(config, EditAnywhere, BlueprintReadWrite, Category = Processing, meta = (ToolTip = "Finds duplicated scene nodes which could be replaced with a single instance"))
 	bool bOptimizeDuplicatedNodes;
 
-	UPROPERTY(config, EditAnywhere, Category = Processing, meta = (ToolTip = "Don't keep nodes that marked invisible in FBX(an din the original scene), except switch variants"))
+	UPROPERTY(config, EditAnywhere, BlueprintReadWrite, Category = Processing, meta = (ToolTip = "Don't keep nodes that marked invisible in FBX(an din the original scene), except switch variants"))
 	bool bRemoveInvisibleNodes;
 
-	UPROPERTY(config, EditAnywhere, Category = Processing, meta = (ToolTip = "Collapse nodes that have identity transform, have no mesh and not used in animation/variants/switches"))
+	UPROPERTY(config, EditAnywhere, BlueprintReadWrite, Category = Processing, meta = (ToolTip = "Collapse nodes that have identity transform, have no mesh and not used in animation/variants/switches"))
 	bool bSimplifyNodeHierarchy;
 
-	UPROPERTY(config, EditAnywhere, Category=AssetImporting, meta=(DisplayName="Import Variants", ToolTip="import VAR files"))
+	UPROPERTY(config, EditAnywhere, BlueprintReadWrite, Category=AssetImporting, meta=(DisplayName="Import Variants", ToolTip="import VAR files"))
 	bool bImportVar;
 
-	UPROPERTY(config, EditAnywhere, Category=AssetImporting, meta=(DisplayName="Variants file path", EditCondition = "bImportVar", ToolTip="Path to the *.var file. By default it will search for a *.var file in the same folder as the FBX file, with the same base filename as it", FilePathFilter="var"))
+	UPROPERTY(config, EditAnywhere, BlueprintReadWrite, Category=AssetImporting, meta=(DisplayName="Variants file path", EditCondition = "bImportVar", ToolTip="Path to the *.var file. By default it will search for a *.var file in the same folder as the FBX file, with the same base filename as it", FilePathFilter="var"))
 	FFilePath VarPath;
 
-	UPROPERTY(config, EditAnywhere, Category=AssetImporting, meta=(DisplayName="Import POS States", ToolTip="import POS files"))
+	UPROPERTY(config, EditAnywhere, BlueprintReadWrite, Category=AssetImporting, meta=(DisplayName="Import POS States", ToolTip="import POS files"))
 	bool bImportPos;
 
-	UPROPERTY(config, EditAnywhere, Category=AssetImporting, meta=(DisplayName="POS file path", EditCondition = "bImportPos", ToolTip="Path to the *.pos file. By default it will search for a *.pos file in the same folder as the FBX file, with the same base filename as it", FilePathFilter="pos"))
+	UPROPERTY(config, EditAnywhere, BlueprintReadWrite, Category=AssetImporting, meta=(DisplayName="POS file path", EditCondition = "bImportPos", ToolTip="Path to the *.pos file. By default it will search for a *.pos file in the same folder as the FBX file, with the same base filename as it", FilePathFilter="pos"))
 	FFilePath PosPath;
 
-	UPROPERTY(config, EditAnywhere, Category=AssetImporting, meta=(DisplayName="Import TML Animations", ToolTip="import TML files"))
+	UPROPERTY(config, EditAnywhere, BlueprintReadWrite, Category=AssetImporting, meta=(DisplayName="Import TML Animations", ToolTip="import TML files"))
 	bool bImportTml;
 
-	UPROPERTY(config, EditAnywhere, Category=AssetImporting, meta=(DisplayName="TML file path", EditCondition = "bImportTml", ToolTip="Path to the *.tml file. By default it will search for a *.tml file in the same folder as the FBX file, with the same base filename as it", FilePathFilter="tml"))
+	UPROPERTY(config, EditAnywhere, BlueprintReadWrite, Category=AssetImporting, meta=(DisplayName="TML file path", EditCondition = "bImportTml", ToolTip="Path to the *.tml file. By default it will search for a *.tml file in the same folder as the FBX file, with the same base filename as it", FilePathFilter="tml"))
 	FFilePath TmlPath;
 
-	UPROPERTY(config, EditAnywhere, Category=AssetImporting, meta=(DisplayName="Shadow Textures", ToolTip="How to handle shadow textures"))
+	UPROPERTY(config, EditAnywhere, BlueprintReadWrite, Category=AssetImporting, meta=(DisplayName="Shadow Textures", ToolTip="How to handle shadow textures"))
 	EShadowTextureMode ShadowTextureMode;
 
 public:
diff --git a/Engine/Plugins/Enterprise/DatasmithFBXImporter/Source/DatasmithDeltaGenTranslator/Private/DatasmithDeltaGenImporter.cpp b/Engine/Plugins/Enterprise/DatasmithFBXImporter/Source/DatasmithDeltaGenTranslator/Private/DatasmithDeltaGenImporter.cpp
--- a/Engine/Plugins/Enterprise/DatasmithFBXImporter/Source/DatasmithDeltaGenTranslator/Private/DatasmithDeltaGenImporter.cpp
+++ b/Engine/Plugins/Enterprise/DatasmithFBXImporter/Source/DatasmithDeltaGenTranslator/Private/DatasmithDeltaGenImporter.cpp
@@ -1122,15 +1122,24 @@ bool FDatasmithDeltaGenImporter::SendSceneToDatasmith()
 		for (FDeltaGenTmlDataTimeline& Timeline : TmlTimelines)
 		{
 			TSharedPtr<IDatasmithLevelSequenceElement> ConvertedSequence = ConvertAnimationTimeline(Timeline);
-			DatasmithScene->AddLevelSequence(ConvertedSequence.ToSharedRef());
+			if (ConvertedSequence.IsValid())
+			{
+				DatasmithScene->AddLevelSequence(ConvertedSequence.ToSharedRef());
+			}
 		}
 
 		TMap<FName, TArray<TSharedPtr<IDatasmithActorElement>>> ImportedActorsByOriginalName;
 		TMap<FName, TSharedPtr<IDatasmithBaseMaterialElement>> ImportedMaterialsByName;
 		BuildAssetMaps(DatasmithScene, ImportedActorsByOriginalName, ImportedMaterialsByName);
 
-		TSharedPtr<IDatasmithLevelVariantSetsElement> LevelVariantSets = FDeltaGenVariantConverter::ConvertVariants(VariantSwitches, PosStates, ImportedActorsByOriginalName, ImportedMaterialsByName);
-		DatasmithScene->AddLevelVariantSets(LevelVariantSets.ToSharedRef());
+		if (ImportOptions->bImportVar)
+		{
+			TSharedPtr<IDatasmithLevelVariantSetsElement> LevelVariantSets = FDeltaGenVariantConverter::ConvertVariants(VariantSwitches, PosStates, ImportedActorsByOriginalName, ImportedMaterialsByName);
+			if (LevelVariantSets.IsValid())
+			{
+				DatasmithScene->AddLevelVariantSets(LevelVariantSets.ToSharedRef());
+			}
+		}
 	}
 	else
 	{
diff --git a/Engine/Plugins/Enterprise/DatasmithFBXImporter/Source/DatasmithFBXTranslator/Public/DatasmithFBXImportOptions.h b/Engine/Plugins/Enterprise/DatasmithFBXImporter/Source/DatasmithFBXTranslator/Public/DatasmithFBXImportOptions.h
--- a/Engine/Plugins/Enterprise/DatasmithFBXImporter/Source/DatasmithFBXTranslator/Public/DatasmithFBXImportOptions.h
+++ b/Engine/Plugins/Enterprise/DatasmithFBXImporter/Source/DatasmithFBXTranslator/Public/DatasmithFBXImportOptions.h
@@ -25,10 +25,10 @@ class DATASMITHFBXTRANSLATOR_API UDatasmithFBXImportOptions : public UObject
 	GENERATED_UCLASS_BODY()
 
 public:
-	UPROPERTY(config, EditAnywhere, Category = Lightmaps, meta = (ToolTip = "Generate new UV coordinates for lightmapping instead of using the highest index UV set. \nTurn this on to have Unreal Studio generate lightmap UV sets automatically.\nTurn this off to try using the highest index existing UV set (if available) as the lightmap UV set.\nFor both cases, geometry without existing UV sets will receive an empty UV set, which will by itself not be valid for use with Lightmass."))
+	UPROPERTY(config, EditAnywhere, BlueprintReadWrite, Category = Lightmaps, meta = (ToolTip = "Generate new UV coordinates for lightmapping instead of using the highest index UV set. \nTurn this on to have Unreal Studio generate lightmap UV sets automatically.\nTurn this off to try using the highest index existing UV set (if available) as the lightmap UV set.\nFor both cases, geometry without existing UV sets will receive an empty UV set, which will by itself not be valid for use with Lightmass."))
 	bool bGenerateLightmapUVs;
 
-	UPROPERTY(config, EditAnywhere, Category=AssetImporting, meta=(DisplayName="Texture folders", ToolTip="Where to look for textures"))
+	UPROPERTY(config, EditAnywhere, BlueprintReadWrite, Category=AssetImporting, meta=(DisplayName="Texture folders", ToolTip="Where to look for textures"))
 	TArray<FDirectoryPath> TextureDirs;
 
 	UPROPERTY(config, EditAnywhere, Category=DebugProperty, meta=(DisplayName="Intermediate Serialization", ToolTip="Cache imported Fbx file in intermediate format for faster debugging"))
diff --git a/Engine/Plugins/Enterprise/DatasmithFBXImporter/Source/DatasmithVREDTranslator/Private/DatasmithVREDImportOptions.h b/Engine/Plugins/Enterprise/DatasmithFBXImporter/Source/DatasmithVREDTranslator/Private/DatasmithVREDImportOptions.h
--- a/Engine/Plugins/Enterprise/DatasmithFBXImporter/Source/DatasmithVREDTranslator/Private/DatasmithVREDImportOptions.h
+++ b/Engine/Plugins/Enterprise/DatasmithFBXImporter/Source/DatasmithVREDTranslator/Private/DatasmithVREDImportOptions.h
@@ -19,37 +19,37 @@ class UDatasmithVREDImportOptions : public UDatasmithFBXImportOptions
 	GENERATED_UCLASS_BODY()
 
 public:
-	UPROPERTY(config, EditAnywhere, Category=Processing, meta=(ToolTip="Merge nodes that are not used in variants/animations"))
+	UPROPERTY(config, EditAnywhere, BlueprintReadWrite, Category=Processing, meta=(ToolTip="Merge nodes that are not used in variants/animations"))
 	bool bMergeNodes;
 
-	UPROPERTY(config, EditAnywhere, Category=Processing, meta=(ToolTip="Finds duplicated scene nodes which could be replaced with a single instance"))
+	UPROPERTY(config, EditAnywhere, BlueprintReadWrite, Category=Processing, meta=(ToolTip="Finds duplicated scene nodes which could be replaced with a single instance"))
 	bool bOptimizeDuplicatedNodes;
 
-	UPROPERTY(config, EditAnywhere, Category= AssetImporting, meta=(DisplayName="Import materials file", ToolTip="Uses the *.mats file saved alongside the exported FBX for a more accurate material reproduction"))
+	UPROPERTY(config, EditAnywhere, BlueprintReadWrite, Category= AssetImporting, meta=(DisplayName="Import materials file", ToolTip="Uses the *.mats file saved alongside the exported FBX for a more accurate material reproduction"))
 	bool bImportMats;
 
-	UPROPERTY(config, EditAnywhere, Category=AssetImporting, meta=(DisplayName="Materials file path", EditCondition = "bImportMats", ToolTip="Path to the *.mats file. By default it will search for a *.mats file in the same folder as the FBX file, with the same base filename as it", FilePathFilter="mats"))
+	UPROPERTY(config, EditAnywhere, BlueprintReadWrite, Category=AssetImporting, meta=(DisplayName="Materials file path", EditCondition = "bImportMats", ToolTip="Path to the *.mats file. By default it will search for a *.mats file in the same folder as the FBX file, with the same base filename as it", FilePathFilter="mats"))
 	FFilePath MatsPath;
 
-	UPROPERTY(config, EditAnywhere, Category= AssetImporting, meta=(DisplayName="Import variants file", ToolTip="Uses the *.var file saved alongside the exported FBX"))
+	UPROPERTY(config, EditAnywhere, BlueprintReadWrite, Category= AssetImporting, meta=(DisplayName="Import variants file", ToolTip="Uses the *.var file saved alongside the exported FBX"))
 	bool bImportVar;
 
-	UPROPERTY(config, EditAnywhere, AdvancedDisplay, Category=AssetImporting, meta=(DisplayName="Clean variants", ToolTip="Removes empty variants, variant sets and invalid options. All discarded items will be logged to console."))
+	UPROPERTY(config, EditAnywhere, AdvancedDisplay, BlueprintReadWrite, Category=AssetImporting, meta=(DisplayName="Clean variants", ToolTip="Removes empty variants, variant sets and invalid options. All discarded items will be logged to console."))
 	bool bCleanVar;
 
-	UPROPERTY(config, EditAnywhere, Category=AssetImporting, meta=(DisplayName="Variants file path", EditCondition = "bImportVar", ToolTip="Path to the *.var file. By default it will search for a *.var file in the same folder as the FBX file, with the same base filename as it", FilePathFilter="var"))
+	UPROPERTY(config, EditAnywhere, BlueprintReadWrite, Category=AssetImporting, meta=(DisplayName="Variants file path", EditCondition = "bImportVar", ToolTip="Path to the *.var file. By default it will search for a *.var file in the same folder as the FBX file, with the same base filename as it", FilePathFilter="var"))
 	FFilePath VarPath;
 
-	UPROPERTY(config, EditAnywhere, Category=AssetImporting, meta=(DisplayName="Import lights file", ToolTip="Uses the *.lights file saved alongside the exported FBX to import extra information about lights not saved in the FBX file"))
+	UPROPERTY(config, EditAnywhere, BlueprintReadWrite, Category=AssetImporting, meta=(DisplayName="Import lights file", ToolTip="Uses the *.lights file saved alongside the exported FBX to import extra information about lights not saved in the FBX file"))
 	bool bImportLightInfo;
 
-	UPROPERTY(config, EditAnywhere, Category=AssetImporting, meta=(DisplayName="Lights file path", EditCondition = "bImportLightInfo", ToolTip="Path to the *.lights file. By default it will search for a *.light file in the same folder as the FBX file, with the same base filename as it", FilePathFilter="lights"))
+	UPROPERTY(config, EditAnywhere, BlueprintReadWrite, Category=AssetImporting, meta=(DisplayName="Lights file path", EditCondition = "bImportLightInfo", ToolTip="Path to the *.lights file. By default it will search for a *.light file in the same folder as the FBX file, with the same base filename as it", FilePathFilter="lights"))
 	FFilePath LightInfoPath;
 
-	UPROPERTY(config, EditAnywhere, Category=AssetImporting, meta=(DisplayName="Import clips file", ToolTip="Uses the *.clips file saved alongside the exported FBX to import information about animation clips and blocks, mirroring VRED's animation system"))
+	UPROPERTY(config, EditAnywhere, BlueprintReadWrite, Category=AssetImporting, meta=(DisplayName="Import clips file", ToolTip="Uses the *.clips file saved alongside the exported FBX to import information about animation clips and blocks, mirroring VRED's animation system"))
 	bool bImportClipInfo;
 
-	UPROPERTY(config, EditAnywhere, Category=AssetImporting, meta=(DisplayName="AnimClips file path", EditCondition = "bImportClipInfo", ToolTip="Path to the *.clips file. By default it will search for a *.clips file in the same folder as the FBX file, with the same base filename as it", FilePathFilter="clips"))
+	UPROPERTY(config, EditAnywhere, BlueprintReadWrite, Category=AssetImporting, meta=(DisplayName="AnimClips file path", EditCondition = "bImportClipInfo", ToolTip="Path to the *.clips file. By default it will search for a *.clips file in the same folder as the FBX file, with the same base filename as it", FilePathFilter="clips"))
 	FFilePath ClipInfoPath;
 
 public:
diff --git a/Engine/Plugins/Enterprise/DatasmithFBXImporter/Source/DatasmithVREDTranslator/Private/DatasmithVREDImporter.cpp b/Engine/Plugins/Enterprise/DatasmithFBXImporter/Source/DatasmithVREDTranslator/Private/DatasmithVREDImporter.cpp
--- a/Engine/Plugins/Enterprise/DatasmithFBXImporter/Source/DatasmithVREDTranslator/Private/DatasmithVREDImporter.cpp
+++ b/Engine/Plugins/Enterprise/DatasmithFBXImporter/Source/DatasmithVREDTranslator/Private/DatasmithVREDImporter.cpp
@@ -2085,14 +2085,17 @@ bool FDatasmithVREDImporter::SendSceneToDatasmith()
 		for (FCombinedAnimBlock& CombinedBlock : CombinedBlocks)
 		{
 			TSharedPtr<IDatasmithLevelSequenceElement> ConvertedBlock = ConvertAnimBlock(CombinedBlock);
-			DatasmithScene->AddLevelSequence(ConvertedBlock.ToSharedRef());
-
-			FImportedAnim NewImportedAnim;
-			NewImportedAnim.Name = FString(ConvertedBlock->GetName());
-			NewImportedAnim.ImportedSequence = ConvertedBlock;
-			NewImportedAnim.OriginalStartSeconds = CombinedBlock.GetStartSeconds();
-			NewImportedAnim.OriginalEndSeconds = CombinedBlock.GetEndSeconds();
-			ImportedAnimElements.Add(NewImportedAnim.Name, NewImportedAnim);
+			if (ConvertedBlock.IsValid())
+			{
+				DatasmithScene->AddLevelSequence(ConvertedBlock.ToSharedRef());
+
+				FImportedAnim NewImportedAnim;
+				NewImportedAnim.Name = FString(ConvertedBlock->GetName());
+				NewImportedAnim.ImportedSequence = ConvertedBlock;
+				NewImportedAnim.OriginalStartSeconds = CombinedBlock.GetStartSeconds();
+				NewImportedAnim.OriginalEndSeconds = CombinedBlock.GetEndSeconds();
+				ImportedAnimElements.Add(NewImportedAnim.Name, NewImportedAnim);
+			}
 		}
 
 		TArray<TSharedPtr<IDatasmithLevelSequenceElement>> CreatedClips;
@@ -2100,16 +2103,25 @@ bool FDatasmithVREDImporter::SendSceneToDatasmith()
 
 		for (const TSharedPtr<IDatasmithLevelSequenceElement>& CreatedClip : CreatedClips)
 		{
-			UE_LOG(LogDatasmithVREDImport, Verbose, TEXT("Added clip level sequence '%s' to the scene"), CreatedClip->GetName());
-			DatasmithScene->AddLevelSequence(CreatedClip.ToSharedRef());
+			if (CreatedClip.IsValid())
+			{
+				UE_LOG(LogDatasmithVREDImport, Verbose, TEXT("Added clip level sequence '%s' to the scene"), CreatedClip->GetName());
+				DatasmithScene->AddLevelSequence(CreatedClip.ToSharedRef());
+			}
 		}
 
 		FActorMap ImportedActorsByOriginalName;
 		FMaterialMap ImportedMaterialsByName;
 		BuildAssetMaps(DatasmithScene, ImportedActorsByOriginalName, ImportedMaterialsByName);
 
-		TSharedPtr<IDatasmithLevelVariantSetsElement> LevelVariantSets = FVREDVariantConverter::ConvertVariants(ParsedVariants, ImportedActorsByOriginalName, ImportedMaterialsByName);
-		DatasmithScene->AddLevelVariantSets(LevelVariantSets.ToSharedRef());
+		if (ImportOptions->bImportVar)
+		{
+			TSharedPtr<IDatasmithLevelVariantSetsElement> LevelVariantSets = FVREDVariantConverter::ConvertVariants(ParsedVariants, ImportedActorsByOriginalName, ImportedMaterialsByName);
+			if (LevelVariantSets.IsValid())
+			{
+				DatasmithScene->AddLevelVariantSets(LevelVariantSets.ToSharedRef());
+			}
+		}
 	}
 	else
 	{
diff --git a/Engine/Plugins/Enterprise/DatasmithGLTFImporter/Source/DatasmithGLTFTranslator/Private/DatasmithGLTFImportOptions.h b/Engine/Plugins/Enterprise/DatasmithGLTFImporter/Source/DatasmithGLTFTranslator/Private/DatasmithGLTFImportOptions.h
--- a/Engine/Plugins/Enterprise/DatasmithGLTFImporter/Source/DatasmithGLTFTranslator/Private/DatasmithGLTFImportOptions.h
+++ b/Engine/Plugins/Enterprise/DatasmithGLTFImporter/Source/DatasmithGLTFTranslator/Private/DatasmithGLTFImportOptions.h
@@ -16,14 +16,14 @@ class UDatasmithGLTFImportOptions : public UObject
 
 public:
 	UPROPERTY(
-	    config, EditAnywhere, Category = Lightmaps,
+	    config, EditAnywhere, BlueprintReadWrite, Category = Lightmaps,
 	    meta =
 	        (ToolTip =
 	             "Generate new UV coordinates for lightmapping instead of using the highest index UV set. \nTurn this on to have Unreal Studio generate lightmap UV sets automatically.\nTurn this off to try using the highest index existing UV set (if available) as the lightmap UV set.\nFor both cases, geometry without existing UV sets will receive an empty UV set, which will by itself not be valid for use with Lightmass."))
 	bool bGenerateLightmapUVs;
 
 	UPROPERTY(
-	    config, EditAnywhere, Category = AssetImporting,
+	    config, EditAnywhere, BlueprintReadWrite, Category = AssetImporting,
 	    meta =
 	        (DisplayName = "Import Uniform Scale",
 	         ToolTip = "Scale factor used for importing assets, by default: 100, for conversion from meters(glTF) to centimeters(Unreal default)."))
diff --git a/Engine/Plugins/Enterprise/VariantManager/Source/VariantManager/Private/SVariantManager.cpp b/Engine/Plugins/Enterprise/VariantManager/Source/VariantManager/Private/SVariantManager.cpp
--- a/Engine/Plugins/Enterprise/VariantManager/Source/VariantManager/Private/SVariantManager.cpp
+++ b/Engine/Plugins/Enterprise/VariantManager/Source/VariantManager/Private/SVariantManager.cpp
@@ -2188,7 +2188,9 @@ FReply SVariantManager::OnAddVariantSetClicked()
 // Will return true if it created or updated a UPropertyValue
 bool AutoCaptureProperty(FVariantManager* VarMan, AActor* TargetActor, const FString& PropertyPath, const UProperty* Property)
 {
-	if (!VarMan || !TargetActor || PropertyPath.IsEmpty())
+	// Transient actors are generated temporarily while dragging actors into the level. Once the
+	// mouse is released, another non-transient actor is instantiated
+	if (!VarMan || !TargetActor || TargetActor->HasAnyFlags(RF_Transient) || PropertyPath.IsEmpty())
 	{
 		return false;
 	}
@@ -2435,7 +2437,8 @@ void SVariantManager::OnObjectPropertyChanged(UObject* Object, struct FPropertyC
 	AActor* TargetActor = nullptr;
 	FString PropertyPath;
 
-	bool bIsBuiltIn = VariantManagerUtils::IsBuiltInStructProperty(Event.MemberProperty);
+	bool bIsStructProperty = Event.MemberProperty && Event.MemberProperty->IsA(UStructProperty::StaticClass());
+	bool bIsBuiltIn = bIsStructProperty && FVariantManagerUtils::IsBuiltInStructProperty(Event.MemberProperty);
 
 	// We don't want to capture just the X component of a RelativeLocation property, but we want to capture
 	// the ISO property of a FPostProcessSettings StructProperty
@@ -2501,8 +2504,8 @@ void SVariantManager::OnObjectPropertyChanged(UObject* Object, struct FPropertyC
 	}
 
 	// If we're a non-built in struct property, build the path with the categories
-	// like the propertycapturer would have done
-	if (Event.MemberProperty && !bIsBuiltIn)
+	// like the propertycapturer would have done (this is mostly to manage Post Process Volume properties)
+	if (bIsStructProperty && !bIsBuiltIn)
 	{
 		// Add 'Settings /'
 		PropertyPath += Event.MemberProperty->GetDisplayNameText().ToString() + PATH_DELIMITER;
@@ -2518,12 +2521,11 @@ void SVariantManager::OnObjectPropertyChanged(UObject* Object, struct FPropertyC
 
 	FString PropertyName = Prop->GetDisplayNameText().ToString();
 	TArray<FString> PropertyPaths;
-	const UProperty* OverrideMaterialsProp = FindField<UProperty>( UMeshComponent::StaticClass(), GET_MEMBER_NAME_CHECKED( UMeshComponent, OverrideMaterials ) );
-	const TSet<FString> ProxyPropertyPaths{ TEXT("Relative Location"), TEXT("Relative Rotation"), TEXT("Relative Scale 3D") };
+	static const TSet<FString> ProxyPropertyPaths{ TEXT("Relative Location"), TEXT("Relative Rotation"), TEXT("Relative Scale 3D") };
 
 	// We capture as just 'Materials' in the Variant Manager UI, instead of 'Override Materials'
 	// Override Materials doesn't work like a regular UArrayProperty, we need to use GetNumMaterials
-	if (Prop == OverrideMaterialsProp)
+	if (Prop == FVariantManagerUtils::GetOverrideMaterialsProperty())
 	{
 		if (UStaticMeshComponent* ObjAsComp = Cast<UStaticMeshComponent>(Object))
 		{
diff --git a/Engine/Plugins/Enterprise/VariantManager/Source/VariantManager/Private/VariantManagerModule.cpp b/Engine/Plugins/Enterprise/VariantManager/Source/VariantManager/Private/VariantManagerModule.cpp
--- a/Engine/Plugins/Enterprise/VariantManager/Source/VariantManager/Private/VariantManagerModule.cpp
+++ b/Engine/Plugins/Enterprise/VariantManager/Source/VariantManager/Private/VariantManagerModule.cpp
@@ -32,6 +32,7 @@
 #include "WorkspaceMenuStructure.h"
 #include "WorkspaceMenuStructureModule.h"
 #include "Widgets/Docking/SDockTab.h"
+#include "VariantManagerUtils.h"
 
 #define LOCTEXT_NAMESPACE "VariantManagerModule"
 
@@ -53,6 +54,9 @@ public:
 			RegisterTabSpawner( LevelEditorModule.GetLevelEditorTabManager() );
 		});
 
+		// Make sure we update the cached UProperty pointers we use for exception properties whenever hot reload happens to a relevant class
+		FVariantManagerUtils::RegisterForHotReload();
+
 		// Register asset actions for the LevelVariantSets asset
 		IAssetTools& AssetTools = FModuleManager::LoadModuleChecked<FAssetToolsModule>(TEXT("AssetTools")).Get();
 		AssetActions = MakeShareable(new FLevelVariantSetsAssetActions(MakeShareable(new FSlateStyleSet(TEXT("EditorStyle")))));
@@ -78,6 +82,8 @@ public:
 			AssetTools.UnregisterAssetTypeActions(AssetActions.ToSharedRef());
 		}
 
+		FVariantManagerUtils::UnregisterForHotReload();
+
 		FVariantManagerEditorCommands::Unregister();
 
 		FPropertyEditorModule& PropertyModule = FModuleManager::LoadModuleChecked< FPropertyEditorModule >(TEXT("PropertyEditor"));
diff --git a/Engine/Plugins/Enterprise/VariantManager/Source/VariantManager/Private/VariantManagerPropertyCapturer.cpp b/Engine/Plugins/Enterprise/VariantManager/Source/VariantManager/Private/VariantManagerPropertyCapturer.cpp
--- a/Engine/Plugins/Enterprise/VariantManager/Source/VariantManager/Private/VariantManagerPropertyCapturer.cpp
+++ b/Engine/Plugins/Enterprise/VariantManager/Source/VariantManager/Private/VariantManagerPropertyCapturer.cpp
@@ -76,17 +76,6 @@ private:
 	TSet<FString> CapturedPaths;
 };
 
-namespace PropertyCaptureHelperNames
-{
-	static const FName RelativeLocationName(TEXT("RelativeLocation"));
-	static const FName RelativeRotationName(TEXT("RelativeRotation"));
-	static const FName RelativeScale3DName(TEXT("RelativeScale3D"));
-	static const FName bVisibleName(TEXT("bVisible"));
-	static const FName OverrideMaterialsName(TEXT("OverrideMaterials"));
-	static const FName LightColorName(TEXT("LightColor"));
-	static const FName DefaultLightColorName(TEXT("DefaultLightColor"));
-}
-
 FPropertyCaptureHelper::FPropertyCaptureHelper(const TArray<UObject*>& InObjects, EPropertyValueCategory InCategoriesToCapture, FString InTargetPropertyPath)
 	: TargetPropertyPath(InTargetPropertyPath)
 	, CategoriesToCapture(InCategoriesToCapture)
@@ -101,181 +90,6 @@ FPropertyCaptureHelper::FPropertyCaptureHelper(const TArray<UObject*>& InObjects
 	TargetPropertyPath.ParseIntoArray(SegmentedTargetPropertyPath, &Delimiter, 1, true);
 }
 
-UArrayProperty* FPropertyCaptureHelper::GetMeshComponentOverloadMaterials()
-{
-	static UArrayProperty* MeshComponentOverloadMaterials = nullptr;
-
-	if (MeshComponentOverloadMaterials == nullptr)
-	{
-		for( TFieldIterator<UArrayProperty> It(UMeshComponent::StaticClass()); It; ++It )
-		{
-			UArrayProperty* Prop = *It;
-			if (Prop->GetFName() == PropertyCaptureHelperNames::OverrideMaterialsName)
-			{
-				MeshComponentOverloadMaterials = Prop;
-				break;
-			}
-		}
-
-		if (MeshComponentOverloadMaterials == nullptr)
-		{
-			UE_LOG(LogVariantManager, Error, TEXT("Unable to find UMeshComponent's OverrideMaterials property!"));
-		}
-	}
-
-	return MeshComponentOverloadMaterials;
-}
-
-UStructProperty* FPropertyCaptureHelper::GetSceneComponentRelativeLocation()
-{
-	static UStructProperty* SceneComponentRelativeLocation = nullptr;
-
-	if (SceneComponentRelativeLocation == nullptr)
-	{
-		for( TFieldIterator<UStructProperty> It(USceneComponent::StaticClass()); It; ++It )
-		{
-			UStructProperty* Prop = *It;
-			if (Prop->GetFName() == PropertyCaptureHelperNames::RelativeLocationName)
-			{
-				SceneComponentRelativeLocation = Prop;
-				break;
-			}
-		}
-
-		if (SceneComponentRelativeLocation == nullptr)
-		{
-			UE_LOG(LogVariantManager, Error, TEXT("Unable to find USceneComponent's RelativeLocation property!"));
-		}
-	}
-
-	return SceneComponentRelativeLocation;
-}
-
-UStructProperty* FPropertyCaptureHelper::GetSceneComponentRelativeRotation()
-{
-	static UStructProperty* SceneComponentRelativeRotation = nullptr;
-
-	if (SceneComponentRelativeRotation == nullptr)
-	{
-		for( TFieldIterator<UStructProperty> It(USceneComponent::StaticClass()); It; ++It )
-		{
-			UStructProperty* Prop = *It;
-			if (Prop->GetFName() == PropertyCaptureHelperNames::RelativeRotationName)
-			{
-				SceneComponentRelativeRotation = Prop;
-				break;
-			}
-		}
-
-		if (SceneComponentRelativeRotation == nullptr)
-		{
-			UE_LOG(LogVariantManager, Error, TEXT("Unable to find USceneComponent's RelativeRotation property!"));
-		}
-	}
-
-	return SceneComponentRelativeRotation;
-}
-
-UStructProperty* FPropertyCaptureHelper::GetSceneComponentRelativeScale3D()
-{
-	static UStructProperty* SceneComponentRelativeScale3D = nullptr;
-
-	if (SceneComponentRelativeScale3D == nullptr)
-	{
-		for( TFieldIterator<UStructProperty> It(USceneComponent::StaticClass()); It; ++It )
-		{
-			UStructProperty* Prop = *It;
-			if (Prop->GetFName() == PropertyCaptureHelperNames::RelativeScale3DName)
-			{
-				SceneComponentRelativeScale3D = Prop;
-				break;
-			}
-		}
-
-		if (SceneComponentRelativeScale3D == nullptr)
-		{
-			UE_LOG(LogVariantManager, Error, TEXT("Unable to find USceneComponent's RelativeScale3D property!"));
-		}
-	}
-
-	return SceneComponentRelativeScale3D;
-}
-
-UBoolProperty* FPropertyCaptureHelper::GetSceneComponentbVisible()
-{
-	static UBoolProperty* SceneComponentbVisible = nullptr;
-
-	if (SceneComponentbVisible == nullptr)
-	{
-		for( TFieldIterator<UBoolProperty> It(USceneComponent::StaticClass()); It; ++It )
-		{
-			UBoolProperty* Prop = *It;
-			if (Prop->GetFName() == PropertyCaptureHelperNames::bVisibleName)
-			{
-				SceneComponentbVisible = Prop;
-				break;
-			}
-		}
-
-		if (SceneComponentbVisible == nullptr)
-		{
-			UE_LOG(LogVariantManager, Error, TEXT("Unable to find USceneComponent's bVisible property!"));
-		}
-	}
-
-	return SceneComponentbVisible;
-}
-
-UStructProperty* FPropertyCaptureHelper::GetLightComponentLightColor()
-{
-	static UStructProperty* LightComponentLightColor = nullptr;
-
-	if (LightComponentLightColor == nullptr)
-	{
-		for( TFieldIterator<UStructProperty> It(ULightComponent::StaticClass()); It; ++It )
-		{
-			UStructProperty* Prop = *It;
-			if (Prop->GetFName() == PropertyCaptureHelperNames::LightColorName)
-			{
-				LightComponentLightColor = Prop;
-				break;
-			}
-		}
-
-		if (LightComponentLightColor == nullptr)
-		{
-			UE_LOG(LogVariantManager, Error, TEXT("Unable to find ULightComponent's LightColor property!"));
-		}
-	}
-
-	return LightComponentLightColor;
-}
-
-UStructProperty* FPropertyCaptureHelper::GetFogComponentDefaultLightColor()
-{
-	static UStructProperty* FogComponentDefaultLightColor = nullptr;
-
-	if (FogComponentDefaultLightColor == nullptr)
-	{
-		for( TFieldIterator<UStructProperty> It(UAtmosphericFogComponent::StaticClass()); It; ++It )
-		{
-			UStructProperty* Prop = *It;
-			if (Prop->GetFName() == PropertyCaptureHelperNames::DefaultLightColorName)
-			{
-				FogComponentDefaultLightColor = Prop;
-				break;
-			}
-		}
-
-		if (FogComponentDefaultLightColor == nullptr)
-		{
-			UE_LOG(LogVariantManager, Error, TEXT("Unable to find UAtmosphericFogComponent's DefaultLightColor property!"));
-		}
-	}
-
-	return FogComponentDefaultLightColor;
-}
-
 void FPropertyCaptureHelper::CaptureActorExceptionProperties(const AActor* Actor, FPropertyPath& PropertyPath, FString& PrettyPathString, TArray<FString>& ComponentNames)
 {
 	if (const ASwitchActor* SwitchActor = Cast<const ASwitchActor>(Actor))
@@ -295,7 +109,7 @@ void FPropertyCaptureHelper::CaptureComponentExceptionProperties(const UActorCom
 			int32 NumMats = ComponentAsMeshComponent->GetNumMaterials();
 			const static FString MatString = FString(TEXT("Material["));
 
-			UArrayProperty* OverrideMatsProp = GetMeshComponentOverloadMaterials();
+			UArrayProperty* OverrideMatsProp = FVariantManagerUtils::GetOverrideMaterialsProperty();
 			PropertyPath.AddProperty(FPropertyInfo(OverrideMatsProp));
 
 			FPropertyInfo LeafInfo(OverrideMatsProp->Inner);
@@ -324,7 +138,7 @@ void FPropertyCaptureHelper::CaptureComponentExceptionProperties(const UActorCom
 
 		if (EnumHasAnyFlags(CategoriesToCapture, EPropertyValueCategory::RelativeLocation))
 		{
-			UStructProperty* RelativeLocationProp = GetSceneComponentRelativeLocation();
+			UStructProperty* RelativeLocationProp = FVariantManagerUtils::GetRelativeLocationProperty();
 
 			PropertyPath.AddProperty(FPropertyInfo(RelativeLocationProp));
 			ComponentNames.Add(FString());
@@ -336,7 +150,7 @@ void FPropertyCaptureHelper::CaptureComponentExceptionProperties(const UActorCom
 
 		if (EnumHasAnyFlags(CategoriesToCapture, EPropertyValueCategory::RelativeRotation))
 		{
-			UStructProperty* RelativeRotationProp = GetSceneComponentRelativeRotation();
+			UStructProperty* RelativeRotationProp = FVariantManagerUtils::GetRelativeRotationProperty();
 
 			PropertyPath.AddProperty(FPropertyInfo(RelativeRotationProp));
 			ComponentNames.Add(FString());
@@ -348,7 +162,7 @@ void FPropertyCaptureHelper::CaptureComponentExceptionProperties(const UActorCom
 
 		if (EnumHasAnyFlags(CategoriesToCapture, EPropertyValueCategory::RelativeScale3D))
 		{
-			UStructProperty* RelativeScale3DProp = GetSceneComponentRelativeScale3D();
+			UStructProperty* RelativeScale3DProp = FVariantManagerUtils::GetRelativeScale3DProperty();
 
 			PropertyPath.AddProperty(FPropertyInfo(RelativeScale3DProp));
 			ComponentNames.Add(FString());
@@ -361,7 +175,7 @@ void FPropertyCaptureHelper::CaptureComponentExceptionProperties(const UActorCom
 
 		if (EnumHasAnyFlags(CategoriesToCapture, EPropertyValueCategory::Visibility))
 		{
-			UBoolProperty* VisibilityProp = GetSceneComponentbVisible();
+			UBoolProperty* VisibilityProp = FVariantManagerUtils::GetVisibilityProperty();
 
 			PropertyPath.AddProperty(FPropertyInfo(VisibilityProp));
 			ComponentNames.Add(FString());
@@ -378,7 +192,7 @@ void FPropertyCaptureHelper::CaptureComponentExceptionProperties(const UActorCom
 
 		if (EnumHasAnyFlags(CategoriesToCapture, EPropertyValueCategory::Color))
 		{
-			UStructProperty* LightColorProp = GetLightComponentLightColor();
+			UStructProperty* LightColorProp = FVariantManagerUtils::GetLightColorProperty();
 
 			PropertyPath.AddProperty(FPropertyInfo(LightColorProp));
 			ComponentNames.Add(FString());
@@ -395,7 +209,7 @@ void FPropertyCaptureHelper::CaptureComponentExceptionProperties(const UActorCom
 
 		if (EnumHasAnyFlags(CategoriesToCapture, EPropertyValueCategory::Color))
 		{
-			UStructProperty* FogComponentProp = GetFogComponentDefaultLightColor();
+			UStructProperty* FogComponentProp = FVariantManagerUtils::GetDefaultLightColorProperty();
 
 			PropertyPath.AddProperty(FPropertyInfo(FogComponentProp));
 			ComponentNames.Add(FString());
@@ -560,7 +374,7 @@ void FPropertyCaptureHelper::GetAllPropertyPathsRecursive(const void* ValuePtr,
 			FString PrettyString = Property->GetDisplayNameText().ToString();
 			if (const UScriptStruct* Struct = Cast<UScriptStruct>(PropertySource))
 			{
-				if (!VariantManagerUtils::IsBuiltInStructProperty(Property))
+				if (!FVariantManagerUtils::IsBuiltInStructProperty(Property))
 				{
 					FString Category = Property->GetMetaData(TEXT("Category"));
 					if (!Category.IsEmpty())
@@ -591,7 +405,7 @@ void FPropertyCaptureHelper::GetAllPropertyPathsRecursive(const void* ValuePtr,
 					{
 						// Check the UArrayProperty's flags. The Inner's flags are unused for this
 						if (CanCaptureProperty(PropertySource, ArrayProperty, PropertySetterName) &&
-							(VariantManagerUtils::IsBuiltInStructProperty(ArrayProperty->Inner) || CAPTURE_ENTIRE_NON_BUILTIN_STRUCTS))
+							(FVariantManagerUtils::IsBuiltInStructProperty(ArrayProperty->Inner) || CAPTURE_ENTIRE_NON_BUILTIN_STRUCTS))
 						{
 							CaptureProp(PropertyPath, PrettyPathString, ComponentNames, PropertySetterName);
 						}
@@ -655,7 +469,7 @@ void FPropertyCaptureHelper::GetAllPropertyPathsRecursive(const void* ValuePtr,
 			// Structs
 			else if (UStructProperty* StructProperty = Cast<UStructProperty>(Property))
 			{
-				if (CanCaptureProperty(PropertySource, Property, PropertySetterName) && (VariantManagerUtils::IsBuiltInStructProperty(Property) || CAPTURE_ENTIRE_NON_BUILTIN_STRUCTS))
+				if (CanCaptureProperty(PropertySource, Property, PropertySetterName) && (FVariantManagerUtils::IsBuiltInStructProperty(Property) || CAPTURE_ENTIRE_NON_BUILTIN_STRUCTS))
 				{
 					CaptureProp(PropertyPath, PrettyPathString, ComponentNames, PropertySetterName);
 				}
diff --git a/Engine/Plugins/Enterprise/VariantManager/Source/VariantManager/Private/VariantManagerUtils.cpp b/Engine/Plugins/Enterprise/VariantManager/Source/VariantManager/Private/VariantManagerUtils.cpp
--- a/Engine/Plugins/Enterprise/VariantManager/Source/VariantManager/Private/VariantManagerUtils.cpp
+++ b/Engine/Plugins/Enterprise/VariantManager/Source/VariantManager/Private/VariantManagerUtils.cpp
@@ -1,30 +1,134 @@
 // Copyright 1998-2019 Epic Games, Inc. All Rights Reserved.
 
 #include "VariantManagerUtils.h"
+
 #include "UObject/UnrealType.h"
+#include "UObject/UObjectGlobals.h"
+#include "Components/SceneComponent.h"
+#include "Components/MeshComponent.h"
+#include "Components/LightComponent.h"
+#include "Atmosphere/AtmosphericFogComponent.h"
+
+UArrayProperty* FVariantManagerUtils::OverrideMaterialsProperty = nullptr;
+UStructProperty* FVariantManagerUtils::RelativeLocationProperty = nullptr;
+UStructProperty* FVariantManagerUtils::RelativeRotationProperty = nullptr;
+UStructProperty* FVariantManagerUtils::RelativeScale3DProperty = nullptr;
+UBoolProperty* FVariantManagerUtils::VisiblityProperty = nullptr;
+UStructProperty* FVariantManagerUtils::LightColorProperty = nullptr;
+UStructProperty* FVariantManagerUtils::DefaultLightColorProperty = nullptr;
+FDelegateHandle FVariantManagerUtils::OnHotReloadHandle;
+
+void FVariantManagerUtils::RegisterForHotReload()
+{
+	OnHotReloadHandle = FCoreUObjectDelegates::RegisterClassForHotReloadReinstancingDelegate.AddStatic(&FVariantManagerUtils::InvalidateCache);
+}
+
+void FVariantManagerUtils::UnregisterForHotReload()
+{
+	FCoreUObjectDelegates::RegisterClassForHotReloadReinstancingDelegate.Remove(OnHotReloadHandle);
+	OnHotReloadHandle.Reset();
+}
+
+bool FVariantManagerUtils::IsBuiltInStructProperty(const UProperty* Property)
+{
+	bool bIsBuiltIn = false;
+
+	const UStructProperty* StructProp = Cast<const UStructProperty>(Property);
+	if (StructProp && StructProp->Struct)
+	{
+		FName StructName = StructProp->Struct->GetFName();
+
+		bIsBuiltIn =
+			StructName == NAME_Rotator ||
+			StructName == NAME_Color ||
+			StructName == NAME_LinearColor ||
+			StructName == NAME_Vector ||
+			StructName == NAME_Quat ||
+			StructName == NAME_Vector4 ||
+			StructName == NAME_Vector2D ||
+			StructName == NAME_IntPoint;
+	}
+
+	return bIsBuiltIn;
+}
 
-namespace VariantManagerUtils
+UArrayProperty* FVariantManagerUtils::GetOverrideMaterialsProperty()
 {
-	bool IsBuiltInStructProperty(const UProperty* Property)
+	if (!OverrideMaterialsProperty)
 	{
-		bool bIsBuiltIn = false;
+		OverrideMaterialsProperty = FindField<UArrayProperty>( UMeshComponent::StaticClass(), GET_MEMBER_NAME_CHECKED( UMeshComponent, OverrideMaterials ) );
+	}
 
-		const UStructProperty* StructProp = Cast<const UStructProperty>(Property);
-		if (StructProp && StructProp->Struct)
-		{
-			FName StructName = StructProp->Struct->GetFName();
+	return OverrideMaterialsProperty;
+}
 
-			bIsBuiltIn =
-				StructName == NAME_Rotator ||
-				StructName == NAME_Color ||
-				StructName == NAME_LinearColor ||
-				StructName == NAME_Vector ||
-				StructName == NAME_Quat ||
-				StructName == NAME_Vector4 ||
-				StructName == NAME_Vector2D ||
-				StructName == NAME_IntPoint;
-		}
+UStructProperty* FVariantManagerUtils::GetRelativeLocationProperty()
+{
+	if (!RelativeLocationProperty)
+	{
+		RelativeLocationProperty = FindField<UStructProperty>( USceneComponent::StaticClass(), GET_MEMBER_NAME_CHECKED( USceneComponent, RelativeLocation ) );
+	}
 
-		return bIsBuiltIn;
+	return RelativeLocationProperty;
+}
+
+UStructProperty* FVariantManagerUtils::GetRelativeRotationProperty()
+{
+	if (!RelativeRotationProperty)
+	{
+		RelativeRotationProperty = FindField<UStructProperty>( USceneComponent::StaticClass(), GET_MEMBER_NAME_CHECKED( USceneComponent, RelativeRotation ) );
 	}
-};
+
+	return RelativeRotationProperty;
+}
+
+UStructProperty* FVariantManagerUtils::GetRelativeScale3DProperty()
+{
+	if (!RelativeScale3DProperty)
+	{
+		RelativeScale3DProperty = FindField<UStructProperty>( USceneComponent::StaticClass(), GET_MEMBER_NAME_CHECKED( USceneComponent, RelativeScale3D ) );
+	}
+
+	return RelativeScale3DProperty;
+}
+
+UBoolProperty* FVariantManagerUtils::GetVisibilityProperty()
+{
+	if (!VisiblityProperty)
+	{
+		VisiblityProperty = FindField<UBoolProperty>( USceneComponent::StaticClass(), GET_MEMBER_NAME_CHECKED( USceneComponent, bVisible ) );
+	}
+
+	return VisiblityProperty;
+}
+
+UStructProperty* FVariantManagerUtils::GetLightColorProperty()
+{
+	if (!LightColorProperty)
+	{
+		LightColorProperty = FindField<UStructProperty>( ULightComponent::StaticClass(), GET_MEMBER_NAME_CHECKED( ULightComponent, LightColor ) );
+	}
+
+	return LightColorProperty;
+}
+
+UStructProperty* FVariantManagerUtils::GetDefaultLightColorProperty()
+{
+	if (!DefaultLightColorProperty)
+	{
+		DefaultLightColorProperty = FindField<UStructProperty>( UAtmosphericFogComponent::StaticClass(), GET_MEMBER_NAME_CHECKED( UAtmosphericFogComponent, DefaultLightColor) );
+	}
+
+	return DefaultLightColorProperty;
+}
+
+void FVariantManagerUtils::InvalidateCache(UClass* OldClass, UClass* NewClass, EHotReloadedClassFlags Flags)
+{
+	OverrideMaterialsProperty = nullptr;
+	RelativeLocationProperty = nullptr;
+	RelativeRotationProperty = nullptr;
+	RelativeScale3DProperty = nullptr;
+	VisiblityProperty = nullptr;
+	LightColorProperty = nullptr;
+	DefaultLightColorProperty = nullptr;
+}
\ No newline at end of file
diff --git a/Engine/Plugins/Enterprise/VariantManager/Source/VariantManager/Private/VariantManagerUtils.h b/Engine/Plugins/Enterprise/VariantManager/Source/VariantManager/Private/VariantManagerUtils.h
--- a/Engine/Plugins/Enterprise/VariantManager/Source/VariantManager/Private/VariantManagerUtils.h
+++ b/Engine/Plugins/Enterprise/VariantManager/Source/VariantManager/Private/VariantManagerUtils.h
@@ -4,7 +4,52 @@
 
 #include "CoreMinimal.h"
 
-namespace VariantManagerUtils
+enum class EHotReloadedClassFlags;
+
+class FVariantManagerUtils
 {
-	bool IsBuiltInStructProperty(const UProperty* Property);
+public:
+	// Invalidate our cached pointers whenever a hot reload happens, as the
+	// classes that own those UPropertys might be reinstanced
+	static void RegisterForHotReload();
+	static void UnregisterForHotReload();
+
+	// Returns true if Property is a UStructProperty with a Struct
+	// of type FVector, FColor, FRotator, FQuat, etc
+	static bool IsBuiltInStructProperty(const UProperty* Property);
+
+	// Returns the OverrideMaterials property of the UMeshComponent class
+	static UArrayProperty* GetOverrideMaterialsProperty();
+
+	// Returns the RelativeLocation property of the USceneComponent class
+	static UStructProperty* GetRelativeLocationProperty();
+
+	// Returns the RelativeRotation property of the USceneComponent class
+	static UStructProperty* GetRelativeRotationProperty();
+
+	// Returns the RelativeScale3D property of the USceneComponent class
+	static UStructProperty* GetRelativeScale3DProperty();
+
+	// Returns the bVisible property of the USceneComponent class
+	static UBoolProperty* GetVisibilityProperty();
+
+	// Returns the LightColor property of the ULightComponent class
+	static UStructProperty* GetLightColorProperty();
+
+	// Returns the DefaultLightColor property of the UAtmosphericFogComponent class
+	static UStructProperty* GetDefaultLightColorProperty();
+
+private:
+	// Invalidates all of our cached UProperty pointers
+	static void InvalidateCache(UClass* OldClass, UClass* NewClass, EHotReloadedClassFlags Flags);
+
+	static UArrayProperty* OverrideMaterialsProperty;
+	static UStructProperty* RelativeLocationProperty;
+	static UStructProperty* RelativeRotationProperty;
+	static UStructProperty* RelativeScale3DProperty;
+	static UBoolProperty* VisiblityProperty;
+	static UStructProperty* LightColorProperty;
+	static UStructProperty* DefaultLightColorProperty;
+
+	static FDelegateHandle OnHotReloadHandle;
 };
\ No newline at end of file
diff --git a/Engine/Plugins/Enterprise/VariantManagerContent/Source/VariantManagerContent/Private/PropertyValue.cpp b/Engine/Plugins/Enterprise/VariantManagerContent/Source/VariantManagerContent/Private/PropertyValue.cpp
--- a/Engine/Plugins/Enterprise/VariantManagerContent/Source/VariantManagerContent/Private/PropertyValue.cpp
+++ b/Engine/Plugins/Enterprise/VariantManagerContent/Source/VariantManagerContent/Private/PropertyValue.cpp
@@ -478,7 +478,7 @@ bool UPropertyValue::Resolve(UObject* Object)
 
 bool UPropertyValue::HasValidResolve() const
 {
-	if (ParentContainerAddress == nullptr)
+	if (ParentContainerAddress == nullptr || ParentContainerClass == nullptr)
 	{
 		return false;
 	}
@@ -487,7 +487,7 @@ bool UPropertyValue::HasValidResolve() const
 	// Whenever an actor's property changes, all of its components that were generated
 	// by construction scripts will be destroyed and regenerated, so we will need to re-resolve
 	// We can't check the resolve state this way for USTRUCTs however (like FPostProcessSettings)
-	if (!ParentContainerClass->IsChildOf(UScriptStruct::StaticClass()))
+	if (!ParentContainerClass->IsA(UScriptStruct::StaticClass()))
 	{
 		if (UObject* Container = (UObject*)ParentContainerAddress)
 		{
diff --git a/Engine/Plugins/Enterprise/VariantManagerContent/Source/VariantManagerContent/Private/PropertyValueColor.cpp b/Engine/Plugins/Enterprise/VariantManagerContent/Source/VariantManagerContent/Private/PropertyValueColor.cpp
--- a/Engine/Plugins/Enterprise/VariantManagerContent/Source/VariantManagerContent/Private/PropertyValueColor.cpp
+++ b/Engine/Plugins/Enterprise/VariantManagerContent/Source/VariantManagerContent/Private/PropertyValueColor.cpp
@@ -16,11 +16,15 @@ UPropertyValueColor::UPropertyValueColor(const FObjectInitializer& ObjectInitial
 {
 }
 
-void UPropertyValueColor::RecordDataFromResolvedObject()
+TArray<uint8> UPropertyValueColor::GetDataFromResolvedObject() const
 {
-	if (!Resolve())
+	int32 PropertySizeBytes = GetValueSizeInBytes();
+	TArray<uint8> CurrentData;
+	CurrentData.SetNumZeroed(PropertySizeBytes);
+
+	if (!HasValidResolve())
 	{
-		return;
+		return CurrentData;
 	}
 
 	// Used by ULightComponent
@@ -30,13 +34,11 @@ void UPropertyValueColor::RecordDataFromResolvedObject()
 		if (!ContainerObject || !ContainerObject->IsValidLowLevel())
 		{
 			UE_LOG(LogVariantContent, Error, TEXT("UPropertyValueColor '%s' does not have a ULightComponent as parent address!"), *GetFullDisplayString());
-			return;
+			return CurrentData;
 		}
 
 		FLinearColor Col = ContainerObject->GetLightColor();
-
-		int32 PropertySizeBytes = GetValueSizeInBytes();
-		SetRecordedData((uint8*)&Col, PropertySizeBytes);
+		FMemory::Memcpy(CurrentData.GetData(), &Col, PropertySizeBytes);
 	}
 	// Used by UAtmosphericFogComponent
 	else if (PropertySetterName == TEXT("SetDefaultLightColor"))
@@ -45,16 +47,14 @@ void UPropertyValueColor::RecordDataFromResolvedObject()
 		if (!ContainerObject || !ContainerObject->IsValidLowLevel())
 		{
 			UE_LOG(LogVariantContent, Error, TEXT("UPropertyValueColor '%s' does not have a UAtmosphericFogComponent as parent address!"), *GetFullDisplayString());
-			return;
+			return CurrentData;
 		}
 
 		FLinearColor Col = FLinearColor(ContainerObject->DefaultLightColor);
-
-		int32 PropertySizeBytes = GetValueSizeInBytes();
-		SetRecordedData((uint8*)&Col, PropertySizeBytes);
+		FMemory::Memcpy(CurrentData.GetData(), &Col, PropertySizeBytes);
 	}
 
-	OnPropertyRecorded.Broadcast();
+	return CurrentData;
 }
 
 UScriptStruct* UPropertyValueColor::GetStructPropertyStruct() const
diff --git a/Engine/Plugins/Enterprise/VariantManagerContent/Source/VariantManagerContent/Private/SwitchActor.cpp b/Engine/Plugins/Enterprise/VariantManagerContent/Source/VariantManagerContent/Private/SwitchActor.cpp
--- a/Engine/Plugins/Enterprise/VariantManagerContent/Source/VariantManagerContent/Private/SwitchActor.cpp
+++ b/Engine/Plugins/Enterprise/VariantManagerContent/Source/VariantManagerContent/Private/SwitchActor.cpp
@@ -8,7 +8,9 @@
 ASwitchActor::ASwitchActor(const FObjectInitializer& Init)
 	: Super(Init)
 {
-	USceneComponent* SceneComponent = CreateDefaultSubobject<USceneComponent>(TEXT("Root"));
+	SceneComponent = CreateDefaultSubobject<USceneComponent>(TEXT("Root"));
+	SceneComponent->SetMobility(EComponentMobility::Static);
+
 	RootComponent = SceneComponent;
 }
 
diff --git a/Engine/Plugins/Enterprise/VariantManagerContent/Source/VariantManagerContent/Public/PropertyValueColor.h b/Engine/Plugins/Enterprise/VariantManagerContent/Source/VariantManagerContent/Public/PropertyValueColor.h
--- a/Engine/Plugins/Enterprise/VariantManagerContent/Source/VariantManagerContent/Public/PropertyValueColor.h
+++ b/Engine/Plugins/Enterprise/VariantManagerContent/Source/VariantManagerContent/Public/PropertyValueColor.h
@@ -18,7 +18,7 @@ class VARIANTMANAGERCONTENT_API UPropertyValueColor : public UPropertyValue
 public:
 
 	// UPropertyValue interface
-	virtual void RecordDataFromResolvedObject() override;
+	TArray<uint8> GetDataFromResolvedObject() const override;
 	virtual UScriptStruct* GetStructPropertyStruct() const override;
 	virtual int32 GetValueSizeInBytes() const override;
 	//~ UPropertyValue interface
diff --git a/Engine/Plugins/Enterprise/VariantManagerContent/Source/VariantManagerContent/Public/SwitchActor.h b/Engine/Plugins/Enterprise/VariantManagerContent/Source/VariantManagerContent/Public/SwitchActor.h
--- a/Engine/Plugins/Enterprise/VariantManagerContent/Source/VariantManagerContent/Public/SwitchActor.h
+++ b/Engine/Plugins/Enterprise/VariantManagerContent/Source/VariantManagerContent/Public/SwitchActor.h
@@ -37,4 +37,8 @@ public:
 
 private:
 	FOnSwitchActorSwitch OnSwitchActorSwitch;
+
+	// Exposing our root component like this allows manual Mobility control on the details panel
+	UPROPERTY(Category = SwitchActor, VisibleAnywhere)
+	class USceneComponent* SceneComponent;
 };
