Copied cl 9471410, 9471857, 9471995, 9487161

VariantManager: Move files to VariantManagerContentEditor
VariantManager: Separate EPropertyValueCategory enum to decouple DatasmithImporter and VariantManager

VariantManager: UI and minor tweaks
- Better way of detecting when we should use the component name or property name for root components;
- Fix lighter gray stripe below actor nodes in 4.24;
- Change auto-expose icon to a standard 'record' icon;

VariantManager: Fix build error for Mac

diff --git a/Engine/Plugins/Enterprise/DatasmithFBXImporter/DatasmithFBXImporter.uplugin b/Engine/Plugins/Enterprise/DatasmithFBXImporter/DatasmithFBXImporter.uplugin
--- a/Engine/Plugins/Enterprise/DatasmithFBXImporter/DatasmithFBXImporter.uplugin
+++ b/Engine/Plugins/Enterprise/DatasmithFBXImporter/DatasmithFBXImporter.uplugin
@@ -41,14 +41,6 @@
         {
 			"Name": "DatasmithContent",
 			"Enabled": true
-		},
-        {
-			"Name": "VariantManager",
-			"Enabled": true
-		},
-        {
-			"Name": "VariantManagerContent",
-			"Enabled": true
-		},
+		}
     ]
 }
diff --git a/Engine/Plugins/Enterprise/DatasmithFBXImporter/Source/DatasmithDeltaGenTranslator/DatasmithDeltaGenTranslator.Build.cs b/Engine/Plugins/Enterprise/DatasmithFBXImporter/Source/DatasmithDeltaGenTranslator/DatasmithDeltaGenTranslator.Build.cs
--- a/Engine/Plugins/Enterprise/DatasmithFBXImporter/Source/DatasmithDeltaGenTranslator/DatasmithDeltaGenTranslator.Build.cs
+++ b/Engine/Plugins/Enterprise/DatasmithFBXImporter/Source/DatasmithDeltaGenTranslator/DatasmithDeltaGenTranslator.Build.cs
@@ -22,7 +22,6 @@ namespace UnrealBuildTool.Rules
 					"MeshDescription",
 					"StaticMeshDescription",
 					"UnrealEd",
-					"VariantManagerContent",
 					"XmlParser",
 				}
 			);
diff --git a/Engine/Plugins/Enterprise/DatasmithFBXImporter/Source/DatasmithDeltaGenTranslator/Private/DatasmithDeltaGenVariantConverter.cpp b/Engine/Plugins/Enterprise/DatasmithFBXImporter/Source/DatasmithDeltaGenTranslator/Private/DatasmithDeltaGenVariantConverter.cpp
--- a/Engine/Plugins/Enterprise/DatasmithFBXImporter/Source/DatasmithDeltaGenTranslator/Private/DatasmithDeltaGenVariantConverter.cpp
+++ b/Engine/Plugins/Enterprise/DatasmithFBXImporter/Source/DatasmithDeltaGenTranslator/Private/DatasmithDeltaGenVariantConverter.cpp
@@ -7,7 +7,6 @@
 #include "DatasmithDeltaGenLog.h"
 #include "DatasmithSceneFactory.h"
 #include "DatasmithVariantElements.h"
-#include "PropertyValue.h"
 
 namespace FDeltaGenVariantConverterImpl
 {
@@ -123,15 +122,15 @@ TSharedPtr<IDatasmithLevelVariantSetsElement> FDeltaGenVariantConverter::Convert
 					FVector Scale = FVector(1.0f, 1.0f, 1.0f);
 
 					TSharedPtr<IDatasmithPropertyCaptureElement> LocationProperty = FDatasmithSceneFactory::CreatePropertyCapture();
-					LocationProperty->SetCategory(EPropertyValueCategory::RelativeLocation);
+					LocationProperty->SetCategory(EDatasmithPropertyCategory::RelativeLocation);
 					LocationProperty->SetRecordedData((uint8*)&Location, sizeof(FVector));
 
 					TSharedPtr<IDatasmithPropertyCaptureElement> RotationProperty = FDatasmithSceneFactory::CreatePropertyCapture();
-					RotationProperty->SetCategory(EPropertyValueCategory::RelativeRotation);
+					RotationProperty->SetCategory(EDatasmithPropertyCategory::RelativeRotation);
 					RotationProperty->SetRecordedData((uint8*)&Rotation, sizeof(FRotator));
 
 					TSharedPtr<IDatasmithPropertyCaptureElement> ScaleProperty = FDatasmithSceneFactory::CreatePropertyCapture();
-					ScaleProperty->SetCategory(EPropertyValueCategory::RelativeScale3D);
+					ScaleProperty->SetCategory(EDatasmithPropertyCategory::RelativeScale3D);
 					ScaleProperty->SetRecordedData((uint8*)&Scale, sizeof(FVector));
 
 					TSharedPtr<IDatasmithActorBindingElement> Binding = FDatasmithSceneFactory::CreateActorBinding();
@@ -180,7 +179,7 @@ TSharedPtr<IDatasmithLevelVariantSetsElement> FDeltaGenVariantConverter::Convert
 				for (const TSharedPtr<IDatasmithActorElement>& NodeActor : NodeActorsArray)
 				{
 					TSharedPtr<IDatasmithPropertyCaptureElement> PropertyCapture = FDatasmithSceneFactory::CreatePropertyCapture();
-					PropertyCapture->SetCategory(EPropertyValueCategory::Visibility);
+					PropertyCapture->SetCategory(EDatasmithPropertyCategory::Visibility);
 
 					FName OriginalName = NodeActor->GetTagsCount() > 0? FName(NodeActor->GetTag(0)) : NAME_None;
 					bool bVisible = VisibileNodes.Contains(OriginalName);
@@ -243,7 +242,7 @@ TSharedPtr<IDatasmithLevelVariantSetsElement> FDeltaGenVariantConverter::Convert
 						case EObjectSetDataType::Translation:
 						{
 							TSharedPtr<IDatasmithPropertyCaptureElement> LocationProperty = FDatasmithSceneFactory::CreatePropertyCapture();
-							LocationProperty->SetCategory(EPropertyValueCategory::RelativeLocation);
+							LocationProperty->SetCategory(EDatasmithPropertyCategory::RelativeLocation);
 							LocationProperty->SetRecordedData((uint8*)Value.Data.GetData(), sizeof(FVector));
 							Prop = LocationProperty;
 							break;
@@ -251,7 +250,7 @@ TSharedPtr<IDatasmithLevelVariantSetsElement> FDeltaGenVariantConverter::Convert
 						case EObjectSetDataType::Rotation:
 						{
 							TSharedPtr<IDatasmithPropertyCaptureElement> RotationProperty = FDatasmithSceneFactory::CreatePropertyCapture();
-							RotationProperty->SetCategory(EPropertyValueCategory::RelativeRotation);
+							RotationProperty->SetCategory(EDatasmithPropertyCategory::RelativeRotation);
 							RotationProperty->SetRecordedData((uint8*)Value.Data.GetData(), sizeof(FRotator));
 							Prop = RotationProperty;
 							break;
@@ -259,7 +258,7 @@ TSharedPtr<IDatasmithLevelVariantSetsElement> FDeltaGenVariantConverter::Convert
 						case EObjectSetDataType::Scaling:
 						{
 							TSharedPtr<IDatasmithPropertyCaptureElement> ScaleProperty = FDatasmithSceneFactory::CreatePropertyCapture();
-							ScaleProperty->SetCategory(EPropertyValueCategory::RelativeScale3D);
+							ScaleProperty->SetCategory(EDatasmithPropertyCategory::RelativeScale3D);
 							ScaleProperty->SetRecordedData((uint8*)Value.Data.GetData(), sizeof(FVector));
 							Prop = ScaleProperty;
 							break;
@@ -267,7 +266,7 @@ TSharedPtr<IDatasmithLevelVariantSetsElement> FDeltaGenVariantConverter::Convert
 						case EObjectSetDataType::Visibility:
 						{
 							TSharedPtr<IDatasmithPropertyCaptureElement> VisibilityProperty = FDatasmithSceneFactory::CreatePropertyCapture();
-							VisibilityProperty->SetCategory(EPropertyValueCategory::Visibility);
+							VisibilityProperty->SetCategory(EDatasmithPropertyCategory::Visibility);
 							VisibilityProperty->SetRecordedData((uint8*)Value.Data.GetData(), sizeof(bool));
 							Prop = VisibilityProperty;
 							break;
@@ -314,7 +313,7 @@ TSharedPtr<IDatasmithLevelVariantSetsElement> FDeltaGenVariantConverter::Convert
 					bool bIsVisible = (ChildIndex == Option.Selection);
 
 					TSharedPtr<IDatasmithPropertyCaptureElement> VisibilityProperty = FDatasmithSceneFactory::CreatePropertyCapture();
-					VisibilityProperty->SetCategory(EPropertyValueCategory::Visibility);
+					VisibilityProperty->SetCategory(EDatasmithPropertyCategory::Visibility);
 					VisibilityProperty->SetRecordedData((uint8*)&bIsVisible, sizeof(bool));
 
 					TSharedPtr<IDatasmithActorBindingElement> Binding = FDatasmithSceneFactory::CreateActorBinding();
@@ -489,7 +488,7 @@ TSharedPtr<IDatasmithLevelVariantSetsElement> FDeltaGenVariantConverter::Convert
 					NewVarSet->AddVariant(NewVar.ToSharedRef());
 
 					TSharedPtr<IDatasmithPropertyCaptureElement> PropertyCapture = FDatasmithSceneFactory::CreatePropertyCapture();
-					PropertyCapture->SetCategory(EPropertyValueCategory::Visibility);
+					PropertyCapture->SetCategory(EDatasmithPropertyCategory::Visibility);
 					PropertyCapture->SetRecordedData((uint8*)&bOn, sizeof(bool));
 
 					TSharedPtr<IDatasmithActorBindingElement> Binding = FDatasmithSceneFactory::CreateActorBinding();
@@ -522,7 +521,7 @@ TSharedPtr<IDatasmithLevelVariantSetsElement> FDeltaGenVariantConverter::Convert
 					bool bIsVisible = (ChildIndex == static_cast<int32>(Selection));
 
 					TSharedPtr<IDatasmithPropertyCaptureElement> VisibilityProperty = FDatasmithSceneFactory::CreatePropertyCapture();
-					VisibilityProperty->SetCategory(EPropertyValueCategory::Visibility);
+					VisibilityProperty->SetCategory(EDatasmithPropertyCategory::Visibility);
 					VisibilityProperty->SetRecordedData((uint8*)&bIsVisible, sizeof(bool));
 
 					TSharedPtr<IDatasmithActorBindingElement> Binding = FDatasmithSceneFactory::CreateActorBinding();
@@ -558,7 +557,7 @@ TSharedPtr<IDatasmithLevelVariantSetsElement> FDeltaGenVariantConverter::Convert
 				}
 
 				TSharedPtr<IDatasmithObjectPropertyCaptureElement> PropertyCapture = FDatasmithSceneFactory::CreateObjectPropertyCapture();
-				PropertyCapture->SetCategory(EPropertyValueCategory::Material);
+				PropertyCapture->SetCategory(EDatasmithPropertyCategory::Material);
 				PropertyCapture->SetRecordedObject(*DatasmithMaterial);
 
 				TSharedPtr<IDatasmithActorBindingElement> Binding = FDatasmithSceneFactory::CreateActorBinding();
diff --git a/Engine/Plugins/Enterprise/DatasmithFBXImporter/Source/DatasmithVREDTranslator/DatasmithVREDTranslator.Build.cs b/Engine/Plugins/Enterprise/DatasmithFBXImporter/Source/DatasmithVREDTranslator/DatasmithVREDTranslator.Build.cs
--- a/Engine/Plugins/Enterprise/DatasmithFBXImporter/Source/DatasmithVREDTranslator/DatasmithVREDTranslator.Build.cs
+++ b/Engine/Plugins/Enterprise/DatasmithFBXImporter/Source/DatasmithVREDTranslator/DatasmithVREDTranslator.Build.cs
@@ -22,7 +22,6 @@ namespace UnrealBuildTool.Rules
 					"MeshDescription",
 					"UnrealEd",
 					"XmlParser",
-					"VariantManagerContent",
 				}
 			);
 
diff --git a/Engine/Plugins/Enterprise/DatasmithFBXImporter/Source/DatasmithVREDTranslator/Private/DatasmithVREDVariantConverter.cpp b/Engine/Plugins/Enterprise/DatasmithFBXImporter/Source/DatasmithVREDTranslator/Private/DatasmithVREDVariantConverter.cpp
--- a/Engine/Plugins/Enterprise/DatasmithFBXImporter/Source/DatasmithVREDTranslator/Private/DatasmithVREDVariantConverter.cpp
+++ b/Engine/Plugins/Enterprise/DatasmithFBXImporter/Source/DatasmithVREDTranslator/Private/DatasmithVREDVariantConverter.cpp
@@ -7,7 +7,6 @@
 #include "DatasmithVariantElements.h"
 #include "DatasmithVREDImportData.h"
 #include "DatasmithVREDLog.h"
-#include "PropertyValue.h"
 
 #define UNGROUPED_VARSET_NAME TEXT("Ungrouped variant sets")
 
@@ -184,7 +183,7 @@ TSharedPtr<IDatasmithLevelVariantSetsElement> FVREDVariantConverter::ConvertVari
 				for (const TSharedPtr<IDatasmithActorElement>& NodeActor : NodeActorsArray)
 				{
 					TSharedPtr<IDatasmithPropertyCaptureElement> PropertyCapture = FDatasmithSceneFactory::CreatePropertyCapture();
-					PropertyCapture->SetCategory(EPropertyValueCategory::Visibility);
+					PropertyCapture->SetCategory(EDatasmithPropertyCategory::Visibility);
 
 					FString OriginalName = NodeActor->GetTagsCount() > 0 ? NodeActor->GetTag(0) : FString();
 					bool bVisible = VisibileNodes.Contains(OriginalName);
@@ -242,7 +241,7 @@ TSharedPtr<IDatasmithLevelVariantSetsElement> FVREDVariantConverter::ConvertVari
 				for (const TSharedPtr<IDatasmithActorElement>& NodeActor : NodeActorsArray)
 				{
 					TSharedPtr<IDatasmithPropertyCaptureElement> PropertyCapture = FDatasmithSceneFactory::CreatePropertyCapture();
-					PropertyCapture->SetCategory(EPropertyValueCategory::Visibility);
+					PropertyCapture->SetCategory(EDatasmithPropertyCategory::Visibility);
 					PropertyCapture->SetRecordedData((uint8*)&bVisible, sizeof(bool));
 
 					TSharedPtr<IDatasmithActorBindingElement> Binding = FDatasmithSceneFactory::CreateActorBinding();
@@ -293,7 +292,7 @@ TSharedPtr<IDatasmithLevelVariantSetsElement> FVREDVariantConverter::ConvertVari
 				for (const TSharedPtr<IDatasmithActorElement>& NodeActor : NodeActorsArray)
 				{
 					TSharedPtr<IDatasmithObjectPropertyCaptureElement> PropertyCapture = FDatasmithSceneFactory::CreateObjectPropertyCapture();
-					PropertyCapture->SetCategory(EPropertyValueCategory::Material);
+					PropertyCapture->SetCategory(EDatasmithPropertyCategory::Material);
 					PropertyCapture->SetRecordedObject(*DatasmithMaterial);
 
 					TSharedPtr<IDatasmithActorBindingElement> Binding = FDatasmithSceneFactory::CreateActorBinding();
@@ -335,15 +334,15 @@ TSharedPtr<IDatasmithLevelVariantSetsElement> FVREDVariantConverter::ConvertVari
 				for (const TSharedPtr<IDatasmithActorElement>& NodeActor : NodeActorsArray)
 				{
 					TSharedPtr<IDatasmithPropertyCaptureElement> LocationProperty = FDatasmithSceneFactory::CreatePropertyCapture();
-					LocationProperty->SetCategory(EPropertyValueCategory::RelativeLocation);
+					LocationProperty->SetCategory(EDatasmithPropertyCategory::RelativeLocation);
 					LocationProperty->SetRecordedData((uint8*)&Loc, sizeof(FVector));
 
 					TSharedPtr<IDatasmithPropertyCaptureElement> RotationProperty = FDatasmithSceneFactory::CreatePropertyCapture();
-					RotationProperty->SetCategory(EPropertyValueCategory::RelativeRotation);
+					RotationProperty->SetCategory(EDatasmithPropertyCategory::RelativeRotation);
 					RotationProperty->SetRecordedData((uint8*)&Rot, sizeof(FRotator));
 
 					TSharedPtr<IDatasmithPropertyCaptureElement> ScaleProperty = FDatasmithSceneFactory::CreatePropertyCapture();
-					ScaleProperty->SetCategory(EPropertyValueCategory::RelativeScale3D);
+					ScaleProperty->SetCategory(EDatasmithPropertyCategory::RelativeScale3D);
 					ScaleProperty->SetRecordedData((uint8*)&Scale, sizeof(FVector));
 
 					TSharedPtr<IDatasmithActorBindingElement> Binding = FDatasmithSceneFactory::CreateActorBinding();
diff --git a/Engine/Plugins/Enterprise/DatasmithImporter/Source/DatasmithImporter/Private/DatasmithLevelVariantSetsImporter.cpp b/Engine/Plugins/Enterprise/DatasmithImporter/Source/DatasmithImporter/Private/DatasmithLevelVariantSetsImporter.cpp
--- a/Engine/Plugins/Enterprise/DatasmithImporter/Source/DatasmithImporter/Private/DatasmithLevelVariantSetsImporter.cpp
+++ b/Engine/Plugins/Enterprise/DatasmithImporter/Source/DatasmithImporter/Private/DatasmithLevelVariantSetsImporter.cpp
@@ -79,7 +79,17 @@ public:
 
 	void ImportPropertyCapture(const TSharedRef<IDatasmithBasePropertyCaptureElement>& BasePropCaptureElement, UVariantObjectBinding* ParentBinding)
 	{
-		EPropertyValueCategory Category = BasePropCaptureElement->GetCategory();
+		static_assert((uint8)EDatasmithPropertyCategory::Undefined == (uint8)EPropertyValueCategory::Undefined, "INVALID_ENUM_VALUE");
+		static_assert((uint8)EDatasmithPropertyCategory::Generic == (uint8)EPropertyValueCategory::Generic, "INVALID_ENUM_VALUE");
+		static_assert((uint8)EDatasmithPropertyCategory::RelativeLocation == (uint8)EPropertyValueCategory::RelativeLocation, "INVALID_ENUM_VALUE");
+		static_assert((uint8)EDatasmithPropertyCategory::RelativeRotation == (uint8)EPropertyValueCategory::RelativeRotation, "INVALID_ENUM_VALUE");
+		static_assert((uint8)EDatasmithPropertyCategory::RelativeScale3D == (uint8)EPropertyValueCategory::RelativeScale3D, "INVALID_ENUM_VALUE");
+		static_assert((uint8)EDatasmithPropertyCategory::Visibility == (uint8)EPropertyValueCategory::Visibility, "INVALID_ENUM_VALUE");
+		static_assert((uint8)EDatasmithPropertyCategory::Material == (uint8)EPropertyValueCategory::Material, "INVALID_ENUM_VALUE");
+		static_assert((uint8)EDatasmithPropertyCategory::Color == (uint8)EPropertyValueCategory::Color, "INVALID_ENUM_VALUE");
+		static_assert((uint8)EDatasmithPropertyCategory::Option == (uint8)EPropertyValueCategory::Option, "INVALID_ENUM_VALUE");
+
+		EDatasmithPropertyCategory Category = BasePropCaptureElement->GetCategory();
 
 		if (BasePropCaptureElement->IsSubType((uint64)EDatasmithElementVariantSubType::PropertyCapture))
 		{
@@ -88,7 +98,7 @@ public:
 
 			UPropertyValue* PropVal = nullptr;
 
-			if (EnumHasAnyFlags(Category, EPropertyValueCategory::Visibility))
+			if (EnumHasAnyFlags(Category, EDatasmithPropertyCategory::Visibility))
 			{
 				TArray<UPropertyValue*> PropertyValues = VariantManager->CreateVisibilityPropertyCaptures({ParentBinding});
 				if (PropertyValues.Num() == 1)
@@ -96,7 +106,7 @@ public:
 					PropVal = PropertyValues[0];
 				}
 			}
-			else if (EnumHasAnyFlags(Category, EPropertyValueCategory::RelativeLocation))
+			else if (EnumHasAnyFlags(Category, EDatasmithPropertyCategory::RelativeLocation))
 			{
 				TArray<UPropertyValue*> PropertyValues = VariantManager->CreateLocationPropertyCaptures({ParentBinding});
 				if (PropertyValues.Num() == 1)
@@ -104,7 +114,7 @@ public:
 					PropVal = PropertyValues[0];
 				}
 			}
-			else if (EnumHasAnyFlags(Category, EPropertyValueCategory::RelativeRotation))
+			else if (EnumHasAnyFlags(Category, EDatasmithPropertyCategory::RelativeRotation))
 			{
 				TArray<UPropertyValue*> PropertyValues = VariantManager->CreateRotationPropertyCaptures({ParentBinding});
 				if (PropertyValues.Num() == 1)
@@ -112,7 +122,7 @@ public:
 					PropVal = PropertyValues[0];
 				}
 			}
-			else if (EnumHasAnyFlags(Category, EPropertyValueCategory::RelativeScale3D))
+			else if (EnumHasAnyFlags(Category, EDatasmithPropertyCategory::RelativeScale3D))
 			{
 				TArray<UPropertyValue*> PropertyValues = VariantManager->CreateScale3DPropertyCaptures({ParentBinding});
 				if (PropertyValues.Num() == 1)
@@ -149,7 +159,7 @@ public:
 
 			UPropertyValue* PropVal = nullptr;
 
-			if (EnumHasAnyFlags(Category, EPropertyValueCategory::Material))
+			if (EnumHasAnyFlags(Category, EDatasmithPropertyCategory::Material))
 			{
 				TArray<UPropertyValue*> PropertyValues = VariantManager->CreateMaterialPropertyCaptures({ParentBinding});
 				if (PropertyValues.Num() == 1)
diff --git a/Engine/Plugins/Enterprise/DatasmithImporter/Source/DatasmithImporter/Private/ObjectElements/DatasmithUObjectElements.cpp b/Engine/Plugins/Enterprise/DatasmithImporter/Source/DatasmithImporter/Private/ObjectElements/DatasmithUObjectElements.cpp
--- a/Engine/Plugins/Enterprise/DatasmithImporter/Source/DatasmithImporter/Private/ObjectElements/DatasmithUObjectElements.cpp
+++ b/Engine/Plugins/Enterprise/DatasmithImporter/Source/DatasmithImporter/Private/ObjectElements/DatasmithUObjectElements.cpp
@@ -2,12 +2,12 @@
 
 #include "ObjectElements/DatasmithUObjectElements.h"
 
+#include "DatasmithDefinitions.h"
 #include "DatasmithSceneFactory.h"
 #include "DatasmithUElementsUtils.h"
 #include "DatasmithUtils.h"
 #include "IDatasmithSceneElements.h"
 #include "ObjectElements/DatasmithUSceneElement.h"
-#include "PropertyValue.h"
 #include "Utility/DatasmithImporterUtils.h"
 
 #include "Engine/StaticMesh.h"
@@ -1073,15 +1073,15 @@ FString UDatasmithBasePropertyCaptureElement::GetPropertyPath() const
 	return Element->GetPropertyPath();
 }
 
-void UDatasmithBasePropertyCaptureElement::SetCategory(EPropertyValueCategory Category)
+void UDatasmithBasePropertyCaptureElement::SetCategory(EDatasmithPropertyCategory Category)
 {
 	DATASMITHOBJECTELEMENT_GETSHARED_AND_EARLYRETURN(IDatasmithBasePropertyCaptureElement, DatasmithElement, );
 	Element->SetCategory(Category);
 }
 
-EPropertyValueCategory UDatasmithBasePropertyCaptureElement::GetCategory() const
+EDatasmithPropertyCategory UDatasmithBasePropertyCaptureElement::GetCategory() const
 {
-	DATASMITHOBJECTELEMENT_GETSHARED_AND_EARLYRETURN(IDatasmithBasePropertyCaptureElement, DatasmithElement, EPropertyValueCategory::Undefined);
+	DATASMITHOBJECTELEMENT_GETSHARED_AND_EARLYRETURN(IDatasmithBasePropertyCaptureElement, DatasmithElement, EDatasmithPropertyCategory::Undefined);
 	return Element->GetCategory();
 }
 
diff --git a/Engine/Plugins/Enterprise/DatasmithImporter/Source/DatasmithImporter/Public/ObjectElements/DatasmithUObjectElements.h b/Engine/Plugins/Enterprise/DatasmithImporter/Source/DatasmithImporter/Public/ObjectElements/DatasmithUObjectElements.h
--- a/Engine/Plugins/Enterprise/DatasmithImporter/Source/DatasmithImporter/Public/ObjectElements/DatasmithUObjectElements.h
+++ b/Engine/Plugins/Enterprise/DatasmithImporter/Source/DatasmithImporter/Public/ObjectElements/DatasmithUObjectElements.h
@@ -5,13 +5,12 @@
 #include "CoreMinimal.h"
 #include "UObject/ObjectMacros.h"
 
+#include "DatasmithDefinitions.h"
 #include "DatasmithVariantElements.h"
 #include "IDatasmithSceneElements.h"
 
 #include "DatasmithUObjectElements.generated.h"
 
-enum class EPropertyValueCategory: uint8;
-
 /*
  * UDatasmithObjectElement
  */
@@ -795,11 +794,11 @@ public:
 
 	/** Sets the category of this property capture */
 	UFUNCTION(BlueprintCallable, Category="Datasmith | Element")
-	void SetCategory(EPropertyValueCategory Category);
+	void SetCategory(EDatasmithPropertyCategory Category);
 
 	/** Gets the category of this property capture */
 	UFUNCTION(BlueprintCallable, Category="Datasmith | Element")
-	EPropertyValueCategory GetCategory() const;
+	EDatasmithPropertyCategory GetCategory() const;
 
 public:
 	TWeakPtr<IDatasmithBasePropertyCaptureElement> GetBasePropertyCaptureElement() const { return DatasmithElement; }
diff --git a/Engine/Plugins/Enterprise/VariantManager/Source/VariantManager/Private/DisplayNodes/VariantManagerDisplayNode.cpp b/Engine/Plugins/Enterprise/VariantManager/Source/VariantManager/Private/DisplayNodes/VariantManagerDisplayNode.cpp
--- a/Engine/Plugins/Enterprise/VariantManager/Source/VariantManager/Private/DisplayNodes/VariantManagerDisplayNode.cpp
+++ b/Engine/Plugins/Enterprise/VariantManager/Source/VariantManager/Private/DisplayNodes/VariantManagerDisplayNode.cpp
@@ -55,7 +55,7 @@ TSharedRef<SWidget> FVariantManagerDisplayNode::GetCustomOutlinerContent(TShared
 		.VAlign(VAlign_Fill)
 		.HAlign(HAlign_Fill)
 		.BorderImage(this, &FVariantManagerDisplayNode::GetNodeBorderImage)
-		.BorderBackgroundColor(FLinearColor(FColor(96, 96, 96, 255)))
+		.BorderBackgroundColor(this, &FVariantManagerDisplayNode::GetNodeBackgroundTint)
 	];
 }
 
diff --git a/Engine/Plugins/Enterprise/VariantManager/Source/VariantManager/Private/LevelVariantSetsActorCustomization.cpp b/Engine/Plugins/Enterprise/VariantManager/Source/VariantManager/Private/LevelVariantSetsActorCustomization.cpp
deleted file mode 100644
diff --git a/Engine/Plugins/Enterprise/VariantManager/Source/VariantManager/Private/LevelVariantSetsActorCustomization.h b/Engine/Plugins/Enterprise/VariantManager/Source/VariantManager/Private/LevelVariantSetsActorCustomization.h
deleted file mode 100644
diff --git a/Engine/Plugins/Enterprise/VariantManager/Source/VariantManager/Private/LevelVariantSetsActorFactory.cpp b/Engine/Plugins/Enterprise/VariantManager/Source/VariantManager/Private/LevelVariantSetsActorFactory.cpp
deleted file mode 100644
diff --git a/Engine/Plugins/Enterprise/VariantManager/Source/VariantManager/Private/LevelVariantSetsAssetActions.cpp b/Engine/Plugins/Enterprise/VariantManager/Source/VariantManager/Private/LevelVariantSetsAssetActions.cpp
deleted file mode 100644
diff --git a/Engine/Plugins/Enterprise/VariantManager/Source/VariantManager/Private/LevelVariantSetsAssetActions.h b/Engine/Plugins/Enterprise/VariantManager/Source/VariantManager/Private/LevelVariantSetsAssetActions.h
deleted file mode 100644
diff --git a/Engine/Plugins/Enterprise/VariantManager/Source/VariantManager/Private/SVariantManager.cpp b/Engine/Plugins/Enterprise/VariantManager/Source/VariantManager/Private/SVariantManager.cpp
--- a/Engine/Plugins/Enterprise/VariantManager/Source/VariantManager/Private/SVariantManager.cpp
+++ b/Engine/Plugins/Enterprise/VariantManager/Source/VariantManager/Private/SVariantManager.cpp
@@ -58,6 +58,7 @@
 #include "Engine/BlueprintGeneratedClass.h"
 #include "VariantManagerUtils.h"
 #include "SwitchActor.h"
+#include "Brushes/SlateImageBrush.h"
 
 #define LOCTEXT_NAMESPACE "SVariantManager"
 
@@ -184,6 +185,8 @@ void SVariantManager::Construct(const FArguments& InArgs, TSharedRef<FVariantMan
 		OnMapChangedHandle = LevelEditorModule.OnMapChanged().AddSP(this, &SVariantManager::OnMapChanged);
 	}
 
+	RecordButtonBrush = MakeShared<FSlateImageBrush>(FPaths::EngineContentDir() / TEXT("Editor/Slate/Icons/CA_Record.png"), FVector2D(24.0f, 24.0f));
+
 	ChildSlot
 	[
 		SAssignNew(MainSplitter, SSplitter)
@@ -240,7 +243,7 @@ void SVariantManager::Construct(const FArguments& InArgs, TSharedRef<FVariantMan
 								.Padding(FMargin(0.0f, 2.0, 2.0f, 2.0)) // Extra padding on the right because ToggleButtonCheckboxes always nudges the image to the right
 								[
 									SNew(SImage)
-									.Image(FEditorStyle::GetBrush("Sequencer.ToggleIsSnapEnabled"))
+									.Image(RecordButtonBrush.Get())
 								]
 							]
 						]
diff --git a/Engine/Plugins/Enterprise/VariantManager/Source/VariantManager/Private/SVariantManager.h b/Engine/Plugins/Enterprise/VariantManager/Source/VariantManager/Private/SVariantManager.h
--- a/Engine/Plugins/Enterprise/VariantManager/Source/VariantManager/Private/SVariantManager.h
+++ b/Engine/Plugins/Enterprise/VariantManager/Source/VariantManager/Private/SVariantManager.h
@@ -26,6 +26,7 @@ class ITableRow;
 class STableViewBase;
 class SSplitter;
 class FTransactionObjectEvent;
+struct FSlateImageBrush;
 enum class EMapChangeType : uint8;
 
 namespace VariantManagerLayoutConstants
@@ -263,4 +264,7 @@ private:
 
 	// We keep track of this to remember splitter values between loads
 	TSharedPtr<SSplitter> MainSplitter;
+	
+	// TODO: Make separate VariantManagerStyle
+	TSharedPtr<FSlateImageBrush> RecordButtonBrush;
 };
diff --git a/Engine/Plugins/Enterprise/VariantManager/Source/VariantManager/Private/SwitchActorCustomization.cpp b/Engine/Plugins/Enterprise/VariantManager/Source/VariantManager/Private/SwitchActorCustomization.cpp
deleted file mode 100644
diff --git a/Engine/Plugins/Enterprise/VariantManager/Source/VariantManager/Private/SwitchActorCustomization.h b/Engine/Plugins/Enterprise/VariantManager/Source/VariantManager/Private/SwitchActorCustomization.h
deleted file mode 100644
diff --git a/Engine/Plugins/Enterprise/VariantManager/Source/VariantManager/Private/VariantManagerBlueprintLibrary.cpp b/Engine/Plugins/Enterprise/VariantManager/Source/VariantManager/Private/VariantManagerBlueprintLibrary.cpp
--- a/Engine/Plugins/Enterprise/VariantManager/Source/VariantManager/Private/VariantManagerBlueprintLibrary.cpp
+++ b/Engine/Plugins/Enterprise/VariantManager/Source/VariantManager/Private/VariantManagerBlueprintLibrary.cpp
@@ -2,19 +2,20 @@
 
 #include "VariantManagerBlueprintLibrary.h"
 
-#include "CoreMinimal.h"
-#include "UObject/ObjectMacros.h"
-#include "VariantManagerModule.h"
-#include "VariantManager.h"
+#include "CapturableProperty.h"
 #include "LevelVariantSets.h"
 #include "LevelVariantSetsActor.h"
-#include "VariantSet.h"
-#include "Variant.h"
-#include "VariantObjectBinding.h"
 #include "PropertyValue.h"
-#include "VariantManagerPropertyCapturer.h"
+#include "Variant.h"
+#include "VariantManager.h"
 #include "VariantManagerLog.h"
-#include "CapturableProperty.h"
+#include "VariantManagerContentEditorModule.h"
+#include "VariantManagerPropertyCapturer.h"
+#include "VariantObjectBinding.h"
+#include "VariantSet.h"
+
+#include "CoreMinimal.h"
+#include "UObject/ObjectMacros.h"
 
 FVariantManager UVariantManagerBlueprintLibrary::VariantManager;
 
@@ -25,8 +26,8 @@ UVariantManagerBlueprintLibrary::UVariantManagerBlueprintLibrary(const FObjectIn
 
 ULevelVariantSets* UVariantManagerBlueprintLibrary::CreateLevelVariantSetsAsset(const FString& AssetName, const FString& AssetPath)
 {
-	IVariantManagerModule& VarManModule = FModuleManager::LoadModuleChecked<IVariantManagerModule>(VARIANTMANAGERMODULE_MODULE_NAME);
-	return Cast<ULevelVariantSets>(VarManModule.CreateLevelVariantSetsAsset(AssetName, AssetPath));
+	IVariantManagerContentEditorModule& ContentEditorModule = FModuleManager::LoadModuleChecked<IVariantManagerContentEditorModule>(VARIANTMANAGERCONTENTEDITORMODULE_MODULE_NAME);
+	return Cast<ULevelVariantSets>(ContentEditorModule.CreateLevelVariantSetsAsset(AssetName, AssetPath));
 }
 
 ALevelVariantSetsActor* UVariantManagerBlueprintLibrary::CreateLevelVariantSetsActor(ULevelVariantSets* LevelVariantSetsAsset)
@@ -36,8 +37,8 @@ ALevelVariantSetsActor* UVariantManagerBlueprintLibrary::CreateLevelVariantSetsA
 		return nullptr;
 	}
 
-	IVariantManagerModule& VarManModule = FModuleManager::LoadModuleChecked<IVariantManagerModule>(VARIANTMANAGERMODULE_MODULE_NAME);
-	return Cast<ALevelVariantSetsActor>(VarManModule.GetOrCreateLevelVariantSetsActor(LevelVariantSetsAsset));
+	IVariantManagerContentEditorModule& ContentEditorModule = FModuleManager::LoadModuleChecked<IVariantManagerContentEditorModule>(VARIANTMANAGERCONTENTEDITORMODULE_MODULE_NAME);
+	return Cast<ALevelVariantSetsActor>(ContentEditorModule.GetOrCreateLevelVariantSetsActor(LevelVariantSetsAsset));
 }
 
 
diff --git a/Engine/Plugins/Enterprise/VariantManager/Source/VariantManager/Private/VariantManagerFactoryNew.cpp b/Engine/Plugins/Enterprise/VariantManager/Source/VariantManager/Private/VariantManagerFactoryNew.cpp
deleted file mode 100644
diff --git a/Engine/Plugins/Enterprise/VariantManager/Source/VariantManager/Private/VariantManagerModule.cpp b/Engine/Plugins/Enterprise/VariantManager/Source/VariantManager/Private/VariantManagerModule.cpp
--- a/Engine/Plugins/Enterprise/VariantManager/Source/VariantManager/Private/VariantManagerModule.cpp
+++ b/Engine/Plugins/Enterprise/VariantManager/Source/VariantManager/Private/VariantManagerModule.cpp
@@ -2,36 +2,21 @@
 
 #include "VariantManagerModule.h"
 
-#include "Modules/ModuleManager.h"
-#include "AssetToolsModule.h"
-#include "ActorFactories/ActorFactory.h"
-#include "CoreMinimal.h"
+#include "LevelVariantSets.h"
+#include "LevelVariantSetsEditorToolkit.h"
+#include "VariantManager.h"
+#include "VariantManagerEditorCommands.h"
+#include "VariantManagerContentEditorModule.h"
+
 #include "Editor.h"
-#include "EditorStyleSet.h"
-#include "Factories/Factory.h"
 #include "Framework/Docking/TabManager.h"
-#include "Framework/MultiBox/MultiBoxBuilder.h"
+#include "IAssetTypeActions.h"
 #include "LevelEditor.h"
-#include "LevelVariantSets.h"
-#include "LevelVariantSetsActor.h"
-#include "LevelVariantSetsAssetActions.h"
+#include "Modules/ModuleManager.h"
 #include "Styling/SlateStyle.h"
-#include "VariantManager.h"
-#include "VariantManagerLog.h"
-#include "VariantManagerEditorCommands.h"
-#include "LevelVariantSetsActorCustomization.h"
-#include "SwitchActorCustomization.h"
-#include "PackageTools.h"
-#include "PropertyEditorModule.h"
-#include "ObjectTools.h"
-#include "PropertyEditorModule.h"
-#include "VariantManagerLog.h"
-#include "LevelVariantSetsEditorToolkit.h"
-#include "UObject/UObjectGlobals.h"
-#include "UObject/ObjectRedirector.h"
+#include "Widgets/Docking/SDockTab.h"
 #include "WorkspaceMenuStructure.h"
 #include "WorkspaceMenuStructureModule.h"
-#include "Widgets/Docking/SDockTab.h"
 #include "VariantManagerUtils.h"
 
 #define LOCTEXT_NAMESPACE "VariantManagerModule"
@@ -45,6 +30,12 @@ public:
 
 		FLevelEditorModule& LevelEditorModule = FModuleManager::LoadModuleChecked<FLevelEditorModule>(TEXT("LevelEditor"));
 
+		// Register a delegate to detect whenever we should open an editor for a LevelVariantSets asset, and relay the
+		// data to LevelVariantSetsEditorToolkit, which will spawn a VariantManager
+		IVariantManagerContentEditorModule& ContentEditorModule = FModuleManager::LoadModuleChecked<IVariantManagerContentEditorModule>(VARIANTMANAGERCONTENTEDITORMODULE_MODULE_NAME);
+		FOnLevelVariantSetsEditor LevelVariantSetsEditorDelegate = FOnLevelVariantSetsEditor::CreateStatic( FVariantManagerModule::OnLevelVariantSetsEditor );
+		ContentEditorModule.RegisterOnLevelVariantSetsDelegate( LevelVariantSetsEditorDelegate );
+
 		// We need to register a tab spawner now so that old tabs that were open when we closed the editor can be reopened
 		// correctly displaying the "Variant Manager" title
 		// Sadly this code runs after the LevelEditorModule is loaded, but before it has created its TabManager. We need to
@@ -56,15 +47,6 @@ public:
 
 		// Make sure we update the cached UProperty pointers we use for exception properties whenever hot reload happens to a relevant class
 		FVariantManagerUtils::RegisterForHotReload();
-
-		// Register asset actions for the LevelVariantSets asset
-		IAssetTools& AssetTools = FModuleManager::LoadModuleChecked<FAssetToolsModule>(TEXT("AssetTools")).Get();
-		AssetActions = MakeShareable(new FLevelVariantSetsAssetActions(MakeShareable(new FSlateStyleSet(TEXT("EditorStyle")))));
-		AssetTools.RegisterAssetTypeActions(AssetActions.ToSharedRef());
-
-		FPropertyEditorModule& PropertyModule = FModuleManager::LoadModuleChecked<FPropertyEditorModule>(TEXT("PropertyEditor"));
-		PropertyModule.RegisterCustomClassLayout(TEXT("LevelVariantSetsActor"), FOnGetDetailCustomizationInstance::CreateStatic(&FLevelVariantSetsActorCustomization::MakeInstance));
-		PropertyModule.RegisterCustomClassLayout(TEXT("SwitchActor"), FOnGetDetailCustomizationInstance::CreateStatic(&FSwitchActorCustomization::MakeInstance));
 	}
 
 	virtual void ShutdownModule() override
@@ -72,23 +54,14 @@ public:
 		FLevelEditorModule& LevelEditorModule = FModuleManager::LoadModuleChecked<FLevelEditorModule>(TEXT("LevelEditor"));
 		LevelEditorModule.OnTabManagerChanged().Remove(OnTabManagerChangedSubscription);
 
-		UnregisterTabSpawner( LevelEditorModule.GetLevelEditorTabManager() );
+		IVariantManagerContentEditorModule& ContentEditorModule = FModuleManager::LoadModuleChecked<IVariantManagerContentEditorModule>(VARIANTMANAGERCONTENTEDITORMODULE_MODULE_NAME);
+		ContentEditorModule.UnregisterOnLevelVariantSetsDelegate();
 
-		// Unregister asset actions for the LevelVariantSets asset
-		FAssetToolsModule* AssetToolsModule = FModuleManager::GetModulePtr<FAssetToolsModule>(TEXT("AssetTools"));
-		if (AssetToolsModule != nullptr)
-		{
-			IAssetTools& AssetTools = AssetToolsModule->Get();
-			AssetTools.UnregisterAssetTypeActions(AssetActions.ToSharedRef());
-		}
+		UnregisterTabSpawner( LevelEditorModule.GetLevelEditorTabManager() );
 
 		FVariantManagerUtils::UnregisterForHotReload();
 
 		FVariantManagerEditorCommands::Unregister();
-
-		FPropertyEditorModule& PropertyModule = FModuleManager::LoadModuleChecked< FPropertyEditorModule >(TEXT("PropertyEditor"));
-		PropertyModule.UnregisterCustomClassLayout(TEXT("LevelVariantSetsActor"));
-		PropertyModule.UnregisterCustomClassLayout(TEXT("SwitchActor"));
 	}
 
 	virtual TSharedRef<FVariantManager> CreateVariantManager(ULevelVariantSets* InLevelVariantSets) override
@@ -99,131 +72,6 @@ public:
 		return VariantManager;
 	}
 
-	virtual UObject* CreateLevelVariantSetsAsset(const FString& AssetName, const FString& PackagePath, bool bForceOverwrite) override
-	{
-		IAssetTools& AssetTools = FModuleManager::GetModuleChecked<FAssetToolsModule>("AssetTools").Get();
-		UObject* NewAsset = nullptr;
-
-		FString SafePackagePath = UPackageTools::SanitizePackageName(PackagePath);
-		FString SafeAssetName = ObjectTools::SanitizeObjectName(AssetName);
-
-		if (SafePackagePath.IsEmpty() || SafeAssetName.IsEmpty())
-		{
-			UE_LOG(LogVariantManager, Error, TEXT("Invalid ULevelVariantSets asset name or package path!"));
-			return nullptr;
-		}
-
-		// Attempt to create a new asset by scanning all available factories until we find one that supports ULevelVariantSets
-		for (TObjectIterator<UClass> It ; It ; ++It)
-		{
-			UClass* CurrentClass = *It;
-			if (CurrentClass->IsChildOf(UFactory::StaticClass()) && !(CurrentClass->HasAnyClassFlags(CLASS_Abstract)))
-			{
-				UFactory* Factory = Cast<UFactory>(CurrentClass->GetDefaultObject());
-				if (Factory->CanCreateNew() && Factory->ImportPriority >= 0 && Factory->SupportedClass == ULevelVariantSets::StaticClass())
-				{
-					// Try deleting existing asset
-					// CreateAsset will always call CanCreateAsset, which will spawn a dialog in case of conflicts.
-					// We want to avoid the dialog, but also still use CreateAsset so that only VariantManagerFactoryNew is responsible for
-					// creating LVS assets. This means we have to manually (try to) delete the asset first
-					// Reference: UAssetToolsImpl::CanCreateAsset
-					if (bForceOverwrite)
-					{
-						const FString PackageName = SafePackagePath + TEXT("/") + SafeAssetName;
-						UPackage* Pkg = CreatePackage(nullptr,*PackageName);
-
-						// Search for UObject instead of ULevelVariantSets as we also want to catch redirectors
-						UObject* ExistingObject = StaticFindObject( UObject::StaticClass(), Pkg, *SafeAssetName );
-
-						if (ExistingObject != nullptr)
-						{
-							// Try to fixup a redirector before we delete it
-							if (ExistingObject->GetClass()->IsChildOf(UObjectRedirector::StaticClass()))
-							{
-								AssetTools.FixupReferencers({Cast<UObjectRedirector>(ExistingObject)});
-							}
-
-							bool bDeleteSucceeded = ObjectTools::DeleteSingleObject( ExistingObject );
-							if (bDeleteSucceeded)
-							{
-								// Force GC so we can cleanly create a new asset (and not do an 'in place' replacement)
-								CollectGarbage( GARBAGE_COLLECTION_KEEPFLAGS );
-
-								// Old package will be GC'ed... create a new one here
-								Pkg = CreatePackage(nullptr,*PackageName);
-								Pkg->MarkAsFullyLoaded();
-							}
-
-							if (!bDeleteSucceeded || !IsUniqueObjectName(*AssetName, Pkg))
-							{
-								UE_LOG(LogVariantManager, Error, TEXT("Failed to delete old ULevelVariantSets asset at '%s'. New asset can't be created."), *PackageName);
-								return nullptr;
-							}
-						}
-					}
-
-					NewAsset = AssetTools.CreateAsset(SafeAssetName, SafePackagePath, ULevelVariantSets::StaticClass(), Factory);
-					break;
-				}
-			}
-		}
-
-		return NewAsset;
-	}
-
-	virtual UObject* CreateLevelVariantSetsAssetWithDialog() override
-	{
-		IAssetTools& AssetTools = FModuleManager::GetModuleChecked<FAssetToolsModule>("AssetTools").Get();
-		UObject* NewAsset = nullptr;
-
-		// Attempt to create a new asset by scanning all available factories until we find one that supports ULevelVariantSets
-		for (TObjectIterator<UClass> It ; It ; ++It)
-		{
-			UClass* CurrentClass = *It;
-			if (CurrentClass->IsChildOf(UFactory::StaticClass()) && !(CurrentClass->HasAnyClassFlags(CLASS_Abstract)))
-			{
-				UFactory* Factory = Cast<UFactory>(CurrentClass->GetDefaultObject());
-				if (Factory->CanCreateNew() && Factory->ImportPriority >= 0 && Factory->SupportedClass == ULevelVariantSets::StaticClass())
-				{
-					NewAsset = AssetTools.CreateAssetWithDialog(ULevelVariantSets::StaticClass(), Factory);
-					break;
-				}
-			}
-		}
-
-		return NewAsset;
-	}
-
-	virtual AActor* GetOrCreateLevelVariantSetsActor(UObject* LevelVariantSetsAsset, bool bForceCreate) override
-	{
-		if (!bForceCreate)
-		{
-			// Check to see if we have a LVSA for this asset already
-			TArray< UObject* > WorldObjects;
-			GetObjectsWithOuter( GWorld, WorldObjects );
-
-			for( UObject* Object : WorldObjects )
-			{
-				if ( ALevelVariantSetsActor* Actor = Cast< ALevelVariantSetsActor >( Object ) )
-				{
-					if ( Actor->GetLevelVariantSets(true) == LevelVariantSetsAsset && !Actor->IsActorBeingDestroyed() )
-					{
-						return Actor;
-					}
-				}
-			}
-		}
-
-		UActorFactory* ActorFactory = GEditor->FindActorFactoryForActorClass(ALevelVariantSetsActor::StaticClass());
-		if (ActorFactory == nullptr)
-		{
-			UE_LOG(LogVariantManager, Error, TEXT("Failed to create LevelVariantSetsActor!"));
-			return nullptr;
-		}
-
-		return GEditor->UseActorFactory(ActorFactory, FAssetData(LevelVariantSetsAsset), &FTransform::Identity);
-	}
-
 	static void RegisterTabSpawner( const TSharedPtr< FTabManager >& TabManager )
 	{
 		if ( !TabManager )
@@ -260,8 +108,14 @@ public:
 			.TabRole(ETabRole::PanelTab);
 	}
 
+	static void OnLevelVariantSetsEditor(const EToolkitMode::Type Mode, const TSharedPtr<class IToolkitHost>& EditWithinLevelEditor, class ULevelVariantSets* LevelVariantSets)
+	{
+		TSharedPtr<ISlateStyle> Style = MakeShareable(new FSlateStyleSet(TEXT("EditorStyle")));
+		TSharedRef<FLevelVariantSetsEditorToolkit> Toolkit = MakeShareable(new FLevelVariantSetsEditorToolkit(Style.ToSharedRef()));
+		Toolkit->Initialize(Mode, EditWithinLevelEditor, LevelVariantSets);
+	}
+
 private:
-	TSharedPtr<IAssetTypeActions> AssetActions;
 	FDelegateHandle OnTabManagerChangedSubscription;
 };
 
diff --git a/Engine/Plugins/Enterprise/VariantManager/Source/VariantManager/Private/VariantManagerPropertyCapturer.cpp b/Engine/Plugins/Enterprise/VariantManager/Source/VariantManager/Private/VariantManagerPropertyCapturer.cpp
--- a/Engine/Plugins/Enterprise/VariantManager/Source/VariantManager/Private/VariantManagerPropertyCapturer.cpp
+++ b/Engine/Plugins/Enterprise/VariantManager/Source/VariantManager/Private/VariantManagerPropertyCapturer.cpp
@@ -61,7 +61,6 @@ private:
 	// Points to the parent instance object. We use this to determine if a component we stumble upon is
 	// on our object's component hierarchy or on some other object's
 	const UObject* CurrentOwnerObject;
-	bool bCurrentOwnerObjectIsBlueprint;
 
 	EPropertyValueCategory CategoriesToCapture;
 
@@ -519,14 +518,15 @@ void FPropertyCaptureHelper::GetAllPropertyPathsRecursive(const void* ValuePtr,
 						ComponentNames[ComponentNames.Num()-1] = NameString;
 
 						FString ComponentPrettyPath;
-						if (bCurrentOwnerObjectIsBlueprint)
+						if (!ObjAsComponent->HasAnyFlags(RF_DefaultSubObject))
 						{
-							// For blueprint actors, the root component is renameable by the user, so let's use that instead
+							// If we don't have a DefaultSubObject tag, we're a generated component for a blueprint, which can be
+							// renamed by the user. So let's use that name instead
 							ComponentPrettyPath = NameString + PATH_DELIMITER;
 						}
 						else
 						{
-							// For regular actors, the root component name can't be renamed by the user and is something like
+							// For DefaultSubObject components, the name can't be set by the user and is something like
 							// StaticMeshComponent0, so we keep the property name, which should be "Static Mesh Component".
 							ComponentPrettyPath = ObjAsComponent->GetClass()->GetDisplayNameText().ToString() + PATH_DELIMITER;
 						}
@@ -579,7 +579,6 @@ TArray<TSharedPtr<FCapturableProperty>>& FPropertyCaptureHelper::CaptureProperti
 			TargetClass = AsClass;
 			TargetInstance = CDO;
 			CurrentOwnerObject = CDO;
-			bCurrentOwnerObjectIsBlueprint = ((UObject*)AsClass)->IsA(UBlueprintGeneratedClass::StaticClass());
 		}
 		// If the passed in object is an instance of a class
 		else
@@ -587,7 +586,6 @@ TArray<TSharedPtr<FCapturableProperty>>& FPropertyCaptureHelper::CaptureProperti
 			TargetClass = TargetObject->GetClass();
 			TargetInstance = TargetObject;
 			CurrentOwnerObject = TargetObject;
-			bCurrentOwnerObjectIsBlueprint = ((UObject*)TargetObject->GetClass())->IsA(UBlueprintGeneratedClass::StaticClass());
 		}
 
 		CaptureActorExceptionProperties((AActor*)TargetInstance, InitialPropertyPath, InitialPrettyString, InitialComponentNames);
diff --git a/Engine/Plugins/Enterprise/VariantManager/Source/VariantManager/Private/VariantManagerTestActor.cpp b/Engine/Plugins/Enterprise/VariantManager/Source/VariantManager/Private/VariantManagerTestActor.cpp
deleted file mode 100644
diff --git a/Engine/Plugins/Enterprise/VariantManager/Source/VariantManager/Public/LevelVariantSetsActorFactory.h b/Engine/Plugins/Enterprise/VariantManager/Source/VariantManager/Public/LevelVariantSetsActorFactory.h
deleted file mode 100644
diff --git a/Engine/Plugins/Enterprise/VariantManager/Source/VariantManager/Public/VariantManagerFactoryNew.h b/Engine/Plugins/Enterprise/VariantManager/Source/VariantManager/Public/VariantManagerFactoryNew.h
deleted file mode 100644
diff --git a/Engine/Plugins/Enterprise/VariantManager/Source/VariantManager/Public/VariantManagerModule.h b/Engine/Plugins/Enterprise/VariantManager/Source/VariantManager/Public/VariantManagerModule.h
--- a/Engine/Plugins/Enterprise/VariantManager/Source/VariantManager/Public/VariantManagerModule.h
+++ b/Engine/Plugins/Enterprise/VariantManager/Source/VariantManager/Public/VariantManagerModule.h
@@ -5,12 +5,11 @@
 #include "Modules/ModuleInterface.h"
 #include "Modules/ModuleManager.h"		// For inline LoadModuleChecked()
 
-#define VARIANTMANAGERMODULE_MODULE_NAME TEXT("VariantManager")
-
-class AActor;
 class FVariantManager;
 class ULevelVariantSets;
 
+#define VARIANTMANAGERMODULE_MODULE_NAME TEXT("VariantManager")
+
 class IVariantManagerModule : public IModuleInterface
 {
 public:
@@ -25,8 +24,5 @@ public:
 	}
 
 	virtual TSharedRef<FVariantManager> CreateVariantManager(ULevelVariantSets* InLevelVariantSets) = 0;
-	virtual UObject* CreateLevelVariantSetsAssetWithDialog() = 0;
-	virtual UObject* CreateLevelVariantSetsAsset(const FString& AssetName, const FString& PackagePath, bool bForceOverwrite = false) = 0;
-	virtual AActor* GetOrCreateLevelVariantSetsActor(UObject* LevelVariantSetsAsset, bool bForceCreate=false) = 0;
 };
 
diff --git a/Engine/Plugins/Enterprise/VariantManager/Source/VariantManager/Public/VariantManagerTestActor.h b/Engine/Plugins/Enterprise/VariantManager/Source/VariantManager/Public/VariantManagerTestActor.h
deleted file mode 100644
diff --git a/Engine/Plugins/Enterprise/VariantManager/Source/VariantManager/VariantManager.Build.cs b/Engine/Plugins/Enterprise/VariantManager/Source/VariantManager/VariantManager.Build.cs
--- a/Engine/Plugins/Enterprise/VariantManager/Source/VariantManager/VariantManager.Build.cs
+++ b/Engine/Plugins/Enterprise/VariantManager/Source/VariantManager/VariantManager.Build.cs
@@ -31,6 +31,7 @@ namespace UnrealBuildTool.Rules
                     "BlueprintGraph", // For function director
 					"WorkspaceMenuStructure",
 					"ToolMenus",
+					"VariantManagerContentEditor",
 				}
 			);
         }
diff --git a/Engine/Plugins/Enterprise/VariantManagerContent/Source/VariantManagerContent/Public/VariantManagerContentModule.h b/Engine/Plugins/Enterprise/VariantManagerContent/Source/VariantManagerContent/Public/VariantManagerContentModule.h
--- a/Engine/Plugins/Enterprise/VariantManagerContent/Source/VariantManagerContent/Public/VariantManagerContentModule.h
+++ b/Engine/Plugins/Enterprise/VariantManagerContent/Source/VariantManagerContent/Public/VariantManagerContentModule.h
@@ -5,7 +5,7 @@
 #include "Modules/ModuleInterface.h"
 #include "Modules/ModuleManager.h"		// For inline LoadModuleChecked()
 
-#define VARIANTMANAGERCONTENTMODULE_MODULE_NAME TEXT("VariantManagerContentModule")
+#define VARIANTMANAGERCONTENTMODULE_MODULE_NAME TEXT("VariantManagerContent")
 
 
 class IVariantManagerContentModule : public IModuleInterface
diff --git a/Engine/Plugins/Enterprise/VariantManagerContent/Source/VariantManagerContentEditor/Private/LevelVariantSetsActorCustomization.cpp b/Engine/Plugins/Enterprise/VariantManagerContent/Source/VariantManagerContentEditor/Private/LevelVariantSetsActorCustomization.cpp
new file mode 100644
--- /dev/null
+++ b/Engine/Plugins/Enterprise/VariantManagerContent/Source/VariantManagerContentEditor/Private/LevelVariantSetsActorCustomization.cpp
@@ -0,0 +1,113 @@
+// Copyright 1998-2019 Epic Games, Inc. All Rights Reserved.
+
+#include "LevelVariantSetsActorCustomization.h"
+
+#include "LevelVariantSets.h"
+#include "LevelVariantSetsActor.h"
+#include "VariantManagerContentEditorLog.h"
+#include "VariantManagerContentEditorModule.h"
+
+#include "AssetToolsModule.h"
+#include "CoreMinimal.h"
+#include "DetailCategoryBuilder.h"
+#include "DetailLayoutBuilder.h"
+#include "DetailWidgetRow.h"
+#include "Editor.h"
+#include "IAssetTools.h"
+#include "Modules/ModuleManager.h"
+#include "Subsystems/AssetEditorSubsystem.h"
+#include "UObject/Object.h"
+#include "Widgets/DeclarativeSyntaxSupport.h"
+#include "Widgets/Input/SButton.h"
+
+#define LOCTEXT_NAMESPACE "LevelVariantSetsActorCustomization"
+
+FLevelVariantSetsActorCustomization::FLevelVariantSetsActorCustomization()
+{
+}
+
+TSharedRef<IDetailCustomization> FLevelVariantSetsActorCustomization::MakeInstance()
+{
+	return MakeShared<FLevelVariantSetsActorCustomization>();
+}
+
+void FLevelVariantSetsActorCustomization::CustomizeDetails(IDetailLayoutBuilder& DetailLayoutBuilder)
+{
+	//Default to standard details panel when multiple variantsselectors are selected at once
+	TArray<TWeakObjectPtr<UObject>> SelectedObjects = DetailLayoutBuilder.GetSelectedObjects();
+	if (SelectedObjects.Num() != 1)
+	{
+		return;
+	}
+
+	TWeakObjectPtr<UObject> SelectedObject = SelectedObjects[0];
+	if (!SelectedObject.IsValid())
+	{
+		return;
+	}
+
+	ALevelVariantSetsActor* Actor = Cast<ALevelVariantSetsActor>(SelectedObject.Get());
+
+	IDetailCategoryBuilder& ActionsCategory = DetailLayoutBuilder.EditCategory(TEXT("VariantManager"));
+
+	ActionsCategory.AddCustomRow(FText())
+	.WholeRowContent()
+	[
+		SNew(SHorizontalBox)
+		+ SHorizontalBox::Slot()
+		[
+			SNew(SButton)
+			.Text(LOCTEXT("OpenVariantManager", "Open Variant Manager"))
+			.OnClicked(this, &FLevelVariantSetsActorCustomization::OnOpenVariantManagerButtonClicked, Actor)
+		]
+	];
+
+	ActionsCategory.AddCustomRow(FText())
+	.WholeRowContent()
+	[
+		SNew(SHorizontalBox)
+		+ SHorizontalBox::Slot()
+		[
+			SNew(SButton)
+			.Text(LOCTEXT("NewLevelVarSets", "Create new LevelVariantSets asset"))
+			.OnClicked(this, &FLevelVariantSetsActorCustomization::OnCreateLevelVarSetsButtonClicked)
+		]
+	];
+
+	ActionsCategory.AddProperty(DetailLayoutBuilder.GetProperty(GET_MEMBER_NAME_CHECKED(ALevelVariantSetsActor, LevelVariantSets)));
+}
+
+FReply FLevelVariantSetsActorCustomization::OnOpenVariantManagerButtonClicked(ALevelVariantSetsActor* Actor)
+{
+	if (Actor == nullptr || !Actor->IsValidLowLevel())
+	{
+		return FReply::Unhandled();
+	}
+
+	ULevelVariantSets* LevelVarSets = Actor->GetLevelVariantSets(true);
+	if (LevelVarSets == nullptr || !LevelVarSets->IsValidLowLevel())
+	{
+		return FReply::Unhandled();
+	}
+
+	GEditor->GetEditorSubsystem<UAssetEditorSubsystem>()->OpenEditorForAssets({LevelVarSets});
+
+	return FReply::Handled();
+}
+
+FReply FLevelVariantSetsActorCustomization::OnCreateLevelVarSetsButtonClicked()
+{
+	IAssetTools& AssetTools = FModuleManager::LoadModuleChecked<FAssetToolsModule>("AssetTools").Get();
+	IVariantManagerContentEditorModule& ContentEditorModule = IVariantManagerContentEditorModule::Get();
+
+	UObject* NewAsset = ContentEditorModule.CreateLevelVariantSetsAssetWithDialog();
+	if (!NewAsset)
+	{
+		UE_LOG(LogVariantManagerContentEditor, Error, TEXT("Did not find an asset factory for a ULevelVariantSets"));
+		return FReply::Unhandled();
+	}
+
+	return FReply::Handled();
+}
+
+#undef LOCTEXT_NAMESPACE
\ No newline at end of file
diff --git a/Engine/Plugins/Enterprise/VariantManagerContent/Source/VariantManagerContentEditor/Private/LevelVariantSetsActorCustomization.h b/Engine/Plugins/Enterprise/VariantManagerContent/Source/VariantManagerContentEditor/Private/LevelVariantSetsActorCustomization.h
new file mode 100644
--- /dev/null
+++ b/Engine/Plugins/Enterprise/VariantManagerContent/Source/VariantManagerContentEditor/Private/LevelVariantSetsActorCustomization.h
@@ -0,0 +1,25 @@
+// Copyright 1998-2019 Epic Games, Inc. All Rights Reserved.
+
+#pragma once
+
+#include "CoreMinimal.h"
+#include "IDetailCustomization.h"
+
+class ALevelVariantSetsActor;
+class FReply;
+class IDetailLayoutBuilder;
+
+class FLevelVariantSetsActorCustomization : public IDetailCustomization
+{
+public:
+	FLevelVariantSetsActorCustomization();
+	static TSharedRef<IDetailCustomization> MakeInstance();
+
+	// IDetailCustomization interface
+	virtual void CustomizeDetails(IDetailLayoutBuilder& DetailLayout) override;
+	// End of IDetailCustomization interface
+
+private:
+	FReply OnOpenVariantManagerButtonClicked(ALevelVariantSetsActor* Actor);
+	FReply OnCreateLevelVarSetsButtonClicked();
+};
diff --git a/Engine/Plugins/Enterprise/VariantManagerContent/Source/VariantManagerContentEditor/Private/LevelVariantSetsActorFactory.cpp b/Engine/Plugins/Enterprise/VariantManagerContent/Source/VariantManagerContentEditor/Private/LevelVariantSetsActorFactory.cpp
new file mode 100644
--- /dev/null
+++ b/Engine/Plugins/Enterprise/VariantManagerContent/Source/VariantManagerContentEditor/Private/LevelVariantSetsActorFactory.cpp
@@ -0,0 +1,61 @@
+// Copyright 1998-2019 Epic Games, Inc. All Rights Reserved.
+
+#include "LevelVariantSetsActorFactory.h"
+
+#include "LevelVariantSets.h"
+#include "LevelVariantSetsActor.h"
+
+#include "ActorFactories/ActorFactory.h"
+#include "AssetData.h"
+
+#define LOCTEXT_NAMESPACE "ALevelVariantSetsActorFactory"
+
+ULevelVariantSetsActorFactory::ULevelVariantSetsActorFactory(const FObjectInitializer& ObjectInitializer)
+	: Super(ObjectInitializer)
+{
+	DisplayName = LOCTEXT("ALevelVariantSetsActorDisplayName", "LevelVariantSetsActor");
+	NewActorClass = ALevelVariantSetsActor::StaticClass();
+}
+
+bool ULevelVariantSetsActorFactory::CanCreateActorFrom( const FAssetData& AssetData, FText& OutErrorMsg )
+{
+	if ( UActorFactory::CanCreateActorFrom( AssetData, OutErrorMsg ) )
+	{
+		return true;
+	}
+
+	if ( AssetData.IsValid() && !AssetData.GetClass()->IsChildOf( ULevelVariantSets::StaticClass() ) )
+	{
+		OutErrorMsg = NSLOCTEXT("CanCreateActor", "NoLevelVariantSetsAsset", "A valid variant sets asset must be specified.");
+		return false;
+	}
+
+	return true;
+}
+
+AActor* ULevelVariantSetsActorFactory::SpawnActor( UObject* Asset, ULevel* InLevel, const FTransform& Transform, EObjectFlags InObjectFlags, const FName Name )
+{
+	ALevelVariantSetsActor* NewActor = Cast<ALevelVariantSetsActor>(Super::SpawnActor(Asset, InLevel, Transform, InObjectFlags, Name));
+
+	if (NewActor)
+	{
+		if (ULevelVariantSets* LevelVariantSets = Cast<ULevelVariantSets>(Asset))
+		{
+			NewActor->SetLevelVariantSets(LevelVariantSets);
+		}
+	}
+
+	return NewActor;
+}
+
+UObject* ULevelVariantSetsActorFactory::GetAssetFromActorInstance(AActor* Instance)
+{
+	if (ALevelVariantSetsActor* LevelVariantSetsActor = Cast<ALevelVariantSetsActor>(Instance))
+	{
+		return LevelVariantSetsActor->LevelVariantSets.TryLoad();
+	}
+
+	return nullptr;
+}
+
+#undef LOCTEXT_NAMESPACE
\ No newline at end of file
diff --git a/Engine/Plugins/Enterprise/VariantManagerContent/Source/VariantManagerContentEditor/Private/LevelVariantSetsAssetActions.cpp b/Engine/Plugins/Enterprise/VariantManagerContent/Source/VariantManagerContentEditor/Private/LevelVariantSetsAssetActions.cpp
new file mode 100644
--- /dev/null
+++ b/Engine/Plugins/Enterprise/VariantManagerContent/Source/VariantManagerContentEditor/Private/LevelVariantSetsAssetActions.cpp
@@ -0,0 +1,91 @@
+// Copyright 1998-2019 Epic Games, Inc. All Rights Reserved.
+
+#include "LevelVariantSetsAssetActions.h"
+
+#include "LevelVariantSets.h"
+#include "VariantManagerContentEditorModule.h"
+
+#include "Engine/Engine.h"
+#include "EngineGlobals.h"
+#include "ToolMenus.h"
+
+#define LOCTEXT_NAMESPACE "LevelVariantSetAssetActions"
+
+FLevelVariantSetsAssetActions::FLevelVariantSetsAssetActions()
+{
+}
+
+uint32 FLevelVariantSetsAssetActions::GetCategories()
+{
+	return EAssetTypeCategories::Misc;
+}
+
+FText FLevelVariantSetsAssetActions::GetName() const
+{
+	return NSLOCTEXT("AssetTypeActions", "AssetTypeActions_LevelVariantSets", "Level Variant Sets");
+}
+
+UClass* FLevelVariantSetsAssetActions::GetSupportedClass() const
+{
+	return ULevelVariantSets::StaticClass();
+}
+
+void FLevelVariantSetsAssetActions::GetActions(const TArray<UObject*>& InObjects, FToolMenuSection& Section)
+{
+	ULevelVariantSets* SomeLevelVarSets = nullptr;
+	for (UObject* Obj : InObjects)
+	{
+		if (ULevelVariantSets* ObjAsLevelVarSets = Cast<ULevelVariantSets>(Obj))
+		{
+			SomeLevelVarSets = ObjAsLevelVarSets;
+			break;
+		}
+	}
+	if (SomeLevelVarSets == nullptr)
+	{
+		return;
+	}
+
+	Section.AddMenuEntry(
+		"CreateActorText",
+		LOCTEXT("CreateActorText", "Create LevelVariantSets actor"),
+		LOCTEXT("CreateActorTooltip", "Creates a new ALevelVariantSetsActor AActor and add it to the scene"),
+		FSlateIcon(),
+		FUIAction(
+			FExecuteAction::CreateLambda([SomeLevelVarSets]()
+			{
+				IVariantManagerContentEditorModule& ContentEditorModule = IVariantManagerContentEditorModule::Get();
+				ContentEditorModule.GetOrCreateLevelVariantSetsActor(SomeLevelVarSets, true);
+			}),
+			FCanExecuteAction()
+		)
+	);
+}
+
+FColor FLevelVariantSetsAssetActions::GetTypeColor() const
+{
+	return FColor(80, 80, 200);
+}
+
+void FLevelVariantSetsAssetActions::OpenAssetEditor(const TArray<UObject*>& InObjects, TSharedPtr<IToolkitHost> EditWithinLevelEditor)
+{
+	EToolkitMode::Type Mode = EditWithinLevelEditor.IsValid()
+		? EToolkitMode::WorldCentric
+		: EToolkitMode::Standalone;
+
+	for (auto ObjIt = InObjects.CreateConstIterator(); ObjIt; ++ObjIt)
+	{
+		if (ULevelVariantSets* LevelVariantSets = Cast<ULevelVariantSets>(*ObjIt))
+		{
+			IVariantManagerContentEditorModule& ContentEditorModule = IVariantManagerContentEditorModule::Get();
+			ContentEditorModule.GetOnLevelVariantSetsEditorOpened().ExecuteIfBound(Mode, EditWithinLevelEditor, LevelVariantSets);
+		}
+	}
+}
+
+bool FLevelVariantSetsAssetActions::ShouldForceWorldCentric()
+{
+	return true;
+}
+
+#undef LOCTEXT_NAMESPACE
diff --git a/Engine/Plugins/Enterprise/VariantManagerContent/Source/VariantManagerContentEditor/Private/LevelVariantSetsAssetActions.h b/Engine/Plugins/Enterprise/VariantManagerContent/Source/VariantManagerContentEditor/Private/LevelVariantSetsAssetActions.h
new file mode 100644
--- /dev/null
+++ b/Engine/Plugins/Enterprise/VariantManagerContent/Source/VariantManagerContentEditor/Private/LevelVariantSetsAssetActions.h
@@ -0,0 +1,26 @@
+// Copyright 1998-2019 Epic Games, Inc. All Rights Reserved.
+
+#pragma once
+
+#include "AssetTypeActions_Base.h"
+#include "CoreMinimal.h"
+#include "Styling/ISlateStyle.h"
+#include "Toolkits/IToolkitHost.h"
+
+class FLevelVariantSetsAssetActions : public FAssetTypeActions_Base
+{
+public:
+	FLevelVariantSetsAssetActions();
+
+	// IAssetTypeActions interface
+	virtual uint32 GetCategories() override;
+	virtual FText GetName() const override;
+	virtual UClass* GetSupportedClass() const override;
+	virtual bool HasActions(const TArray<UObject*>& InObjects) const override { return true; }
+	virtual void GetActions(const TArray<UObject*>& InObjects, struct FToolMenuSection& Section) override;
+	virtual FColor GetTypeColor() const override;
+	virtual void OpenAssetEditor(const TArray<UObject*>& InObjects, TSharedPtr<class IToolkitHost> EditWithinLevelEditor = TSharedPtr<IToolkitHost>()) override;
+	virtual bool ShouldForceWorldCentric() override;
+	virtual bool CanLocalize() const override { return false; }
+	//~ End IAssetTypeActions interface
+};
diff --git a/Engine/Plugins/Enterprise/VariantManagerContent/Source/VariantManagerContentEditor/Private/SwitchActorCustomization.cpp b/Engine/Plugins/Enterprise/VariantManagerContent/Source/VariantManagerContentEditor/Private/SwitchActorCustomization.cpp
new file mode 100644
--- /dev/null
+++ b/Engine/Plugins/Enterprise/VariantManagerContent/Source/VariantManagerContentEditor/Private/SwitchActorCustomization.cpp
@@ -0,0 +1,183 @@
+// Copyright 1998-2019 Epic Games, Inc. All Rights Reserved.
+
+#include "SwitchActorCustomization.h"
+
+#include "SwitchActor.h"
+
+#include "Algo/Transform.h"
+#include "AssetToolsModule.h"
+#include "CoreMinimal.h"
+#include "DetailCategoryBuilder.h"
+#include "DetailLayoutBuilder.h"
+#include "DetailWidgetRow.h"
+#include "Editor.h"
+#include "IAssetTools.h"
+#include "Modules/ModuleManager.h"
+#include "ScopedTransaction.h"
+#include "Toolkits/AssetEditorManager.h"
+#include "UObject/Object.h"
+#include "Widgets/DeclarativeSyntaxSupport.h"
+#include "Widgets/Input/SComboBox.h"
+
+#define LOCTEXT_NAMESPACE "SwitchActorCustomization"
+
+FSwitchActorCustomization::FSwitchActorCustomization()
+{
+}
+
+TSharedRef<IDetailCustomization> FSwitchActorCustomization::MakeInstance()
+{
+	return MakeShared<FSwitchActorCustomization>();
+}
+
+void FSwitchActorCustomization::CustomizeDetails(IDetailLayoutBuilder& DetailLayoutBuilder)
+{
+	//Default to standard details panel when multiple switch actors are selected at once
+	TArray<TWeakObjectPtr<UObject>> SelectedObjects = DetailLayoutBuilder.GetSelectedObjects();
+	if (SelectedObjects.Num() != 1)
+	{
+		return;
+	}
+
+	TWeakObjectPtr<UObject> SelectedObject = SelectedObjects[0];
+	if (!SelectedObject.IsValid())
+	{
+		return;
+	}
+
+	CurrentActor = Cast<ASwitchActor>(SelectedObject.Get());
+	CurrentActor->GetOnSwitchDelegate().AddSP(this, &FSwitchActorCustomization::ForceRefreshDetails);
+
+	TArray<AActor*> ActorOptions = CurrentActor->GetOptions();
+
+	Options.Reset();
+	Algo::Transform(ActorOptions, Options, [](const AActor* Actor)
+	{
+		return MakeShared<FString>(Actor->GetActorLabel());
+	});
+
+	TSharedPtr<FString>* InitiallySelected = nullptr;
+	int32 SelectedIndex = CurrentActor->GetSelectedOption();
+	if (Options.IsValidIndex(SelectedIndex))
+	{
+		InitiallySelected = &Options[SelectedIndex];
+	}
+
+	IDetailCategoryBuilder& ActionsCategory = DetailLayoutBuilder.EditCategory(TEXT("SwitchActor"));
+	const FText FilterStr(LOCTEXT("SelectedOptionRow", "Selected Option"));
+	ActionsCategory.AddCustomRow(FilterStr)
+	.NameContent()
+	[
+		SNew(STextBlock)
+		.Text(FilterStr)
+		.ToolTipText(LOCTEXT("SelectedOptionTooltip", "Which child actor is visible"))
+	]
+	.ValueContent()
+	.MaxDesiredWidth(500.0f)
+	.MinDesiredWidth(100.0f)
+	[
+		SAssignNew(ComboBox, SComboBox<TSharedPtr<FString>>)
+		.OptionsSource(&Options)
+		.InitiallySelectedItem(InitiallySelected? *InitiallySelected : nullptr)
+		.OnComboBoxOpening(this, &FSwitchActorCustomization::OnComboBoxOpening)
+		.OnSelectionChanged(this, &FSwitchActorCustomization::OnComboBoxOptionChanged)
+		.OnGenerateWidget_Lambda([](TSharedPtr<FString> Item)
+		{
+			return SNew(STextBlock).Text(Item.IsValid() ? FText::FromString(*Item) : FText::GetEmpty());
+		})
+		.Content()
+		[
+			SNew(STextBlock).Text(this, &FSwitchActorCustomization::GetComboBoxSelectedOptionText)
+		]
+	];
+}
+
+void FSwitchActorCustomization::CustomizeDetails(const TSharedPtr<IDetailLayoutBuilder>& DetailBuilder)
+{
+	DetailBuilderWeakPtr = DetailBuilder;
+	CustomizeDetails( *DetailBuilder );
+}
+
+void FSwitchActorCustomization::OnComboBoxOpening()
+{
+	if (CurrentActor == nullptr)
+	{
+		return;
+	}
+
+	TArray<AActor*> Actors = CurrentActor->GetOptions();
+
+	Options.Reset();
+	Algo::Transform(Actors, Options, [](const AActor* Actor)
+	{
+		return MakeShared<FString>(Actor->GetActorLabel());
+	});
+
+	ComboBox->ClearSelection();
+	ComboBox->RefreshOptions();
+
+	// Need this to prevent it from clearing the selected option,
+	// as we'll need to rebuild the Options array to make sure its
+	// up to date
+	int32 SelectedOption = CurrentActor->GetSelectedOption();
+	if (Options.IsValidIndex(SelectedOption))
+	{
+		ComboBox->SetSelectedItem(Options[SelectedOption]);
+	}
+}
+
+void FSwitchActorCustomization::OnComboBoxOptionChanged(TSharedPtr<FString> NewOption, ESelectInfo::Type SelectType)
+{
+	// Checking SelectType prevents a refresh loop when setting the default selection in OnComboBoxOpening
+	if (CurrentActor == nullptr || !NewOption.IsValid() || SelectType == ESelectInfo::Direct)
+	{
+		return;
+	}
+
+	FScopedTransaction Transaction(FText::Format(
+		LOCTEXT("OptionChangedTransaction", "Switched '{0}' to option '{1}'"),
+		FText::FromString(CurrentActor->GetActorLabel()),
+		FText::FromString(NewOption.IsValid()? *NewOption : TEXT("None")))
+	);
+
+	// NewOption is newly generated, so it won't match by address of the pointers stored in Options
+	int32 SelectedIndex = INDEX_NONE;
+	for (int32 Index = 0; Index < Options.Num(); ++Index)
+	{
+		const TSharedPtr<FString>& Option = Options[Index];
+		if (Option.IsValid() && *Option == *NewOption)
+		{
+			SelectedIndex = Index;
+			break;
+		}
+	}
+
+	CurrentActor->SelectOption(SelectedIndex);
+
+	// Trick to force the viewport to update even with 'Realtime' turned off,
+	// even though our selection may will remain the same
+	GEditor->NoteSelectionChange();
+}
+
+FText FSwitchActorCustomization::GetComboBoxSelectedOptionText() const
+{
+	TSharedPtr<FString> SelectedItem = ComboBox->GetSelectedItem();
+	if (SelectedItem.IsValid())
+	{
+		return FText::FromString(*SelectedItem);
+	}
+
+	return FText::GetEmpty();
+}
+
+void FSwitchActorCustomization::ForceRefreshDetails(int32 NewOption)
+{
+	// Raw because we don't want to keep alive the details builder when calling the force refresh details
+	IDetailLayoutBuilder* DetailLayoutBuilder = DetailBuilderWeakPtr.Pin().Get();
+	if ( DetailLayoutBuilder )
+	{
+		DetailLayoutBuilder->ForceRefreshDetails();
+	}
+}
+
+#undef LOCTEXT_NAMESPACE
diff --git a/Engine/Plugins/Enterprise/VariantManagerContent/Source/VariantManagerContentEditor/Private/SwitchActorCustomization.h b/Engine/Plugins/Enterprise/VariantManagerContent/Source/VariantManagerContentEditor/Private/SwitchActorCustomization.h
new file mode 100644
--- /dev/null
+++ b/Engine/Plugins/Enterprise/VariantManagerContent/Source/VariantManagerContentEditor/Private/SwitchActorCustomization.h
@@ -0,0 +1,33 @@
+// Copyright 1998-2019 Epic Games, Inc. All Rights Reserved.
+
+#pragma once
+
+#include "CoreMinimal.h"
+#include "IDetailCustomization.h"
+#include "Widgets/Input/SComboBox.h"
+
+class ASwitchActor;
+class IDetailLayoutBuilder;
+
+class FSwitchActorCustomization : public IDetailCustomization
+{
+public:
+	FSwitchActorCustomization();
+	static TSharedRef<IDetailCustomization> MakeInstance();
+
+	// IDetailCustomization interface
+	virtual void CustomizeDetails(IDetailLayoutBuilder& DetailLayout) override;
+	virtual void CustomizeDetails(const TSharedPtr<IDetailLayoutBuilder>& DetailBuilder) override;
+	// End of IDetailCustomization interface
+
+private:
+	void OnComboBoxOpening();
+	void OnComboBoxOptionChanged(TSharedPtr<FString> NewOption, ESelectInfo::Type SelectType);
+	FText GetComboBoxSelectedOptionText() const;
+	void ForceRefreshDetails(int32 NewOption);
+
+	ASwitchActor* CurrentActor;
+	TWeakPtr<IDetailLayoutBuilder> DetailBuilderWeakPtr;
+	TSharedPtr<SComboBox<TSharedPtr<FString>>> ComboBox;
+	TArray<TSharedPtr<FString>> Options;
+};
diff --git a/Engine/Plugins/Enterprise/VariantManagerContent/Source/VariantManagerContentEditor/Private/VariantManagerContentEditorLog.h b/Engine/Plugins/Enterprise/VariantManagerContent/Source/VariantManagerContentEditor/Private/VariantManagerContentEditorLog.h
new file mode 100644
--- /dev/null
+++ b/Engine/Plugins/Enterprise/VariantManagerContent/Source/VariantManagerContentEditor/Private/VariantManagerContentEditorLog.h
@@ -0,0 +1,7 @@
+// Copyright 1998-2019 Epic Games, Inc. All Rights Reserved.
+
+#pragma once
+
+#include "CoreMinimal.h"
+
+DECLARE_LOG_CATEGORY_EXTERN(LogVariantManagerContentEditor, Log, All);
\ No newline at end of file
diff --git a/Engine/Plugins/Enterprise/VariantManagerContent/Source/VariantManagerContentEditor/Private/VariantManagerContentEditorModule.cpp b/Engine/Plugins/Enterprise/VariantManagerContent/Source/VariantManagerContentEditor/Private/VariantManagerContentEditorModule.cpp
new file mode 100644
--- /dev/null
+++ b/Engine/Plugins/Enterprise/VariantManagerContent/Source/VariantManagerContentEditor/Private/VariantManagerContentEditorModule.cpp
@@ -0,0 +1,203 @@
+// Copyright 1998-2019 Epic Games, Inc. All Rights Reserved.
+
+#include "VariantManagerContentEditorModule.h"
+
+#include "LevelVariantSets.h"
+#include "LevelVariantSetsActor.h"
+#include "LevelVariantSetsActorCustomization.h"
+#include "LevelVariantSetsAssetActions.h"
+#include "SwitchActorCustomization.h"
+#include "VariantManagerContentEditorLog.h"
+
+#include "ActorFactories/ActorFactory.h"
+#include "AssetToolsModule.h"
+#include "CoreMinimal.h"
+#include "Factories/Factory.h"
+#include "Modules/ModuleManager.h"
+#include "ObjectTools.h"
+#include "PackageTools.h"
+#include "PropertyEditorModule.h"
+
+#define LOCTEXT_NAMESPACE "VariantManagerContentEditorModule"
+
+DEFINE_LOG_CATEGORY(LogVariantManagerContentEditor);
+
+class FVariantManagerContentEditorModule : public IVariantManagerContentEditorModule
+{
+public:
+	virtual void StartupModule() override
+	{
+		// Register asset actions for the LevelVariantSets asset
+		IAssetTools& AssetTools = FModuleManager::LoadModuleChecked<FAssetToolsModule>(TEXT("AssetTools")).Get();
+		AssetActions = MakeShareable(new FLevelVariantSetsAssetActions());
+		AssetTools.RegisterAssetTypeActions(AssetActions.ToSharedRef());
+
+		FPropertyEditorModule& PropertyModule = FModuleManager::LoadModuleChecked<FPropertyEditorModule>(TEXT("PropertyEditor"));
+		PropertyModule.RegisterCustomClassLayout(TEXT("LevelVariantSetsActor"), FOnGetDetailCustomizationInstance::CreateStatic(&FLevelVariantSetsActorCustomization::MakeInstance));
+		PropertyModule.RegisterCustomClassLayout(TEXT("SwitchActor"), FOnGetDetailCustomizationInstance::CreateStatic(&FSwitchActorCustomization::MakeInstance));
+	}
+
+	virtual void ShutdownModule() override
+	{
+        // Unregister asset actions for the LevelVariantSets asset
+		FAssetToolsModule* AssetToolsModule = FModuleManager::GetModulePtr<FAssetToolsModule>(TEXT("AssetTools"));
+		if (AssetToolsModule != nullptr)
+		{
+			IAssetTools& AssetTools = AssetToolsModule->Get();
+			AssetTools.UnregisterAssetTypeActions(AssetActions.ToSharedRef());
+		}
+
+        FPropertyEditorModule& PropertyModule = FModuleManager::LoadModuleChecked< FPropertyEditorModule >(TEXT("PropertyEditor"));
+		PropertyModule.UnregisterCustomClassLayout(TEXT("LevelVariantSetsActor"));
+		PropertyModule.UnregisterCustomClassLayout(TEXT("SwitchActor"));
+	}
+
+	virtual void RegisterOnLevelVariantSetsDelegate(FOnLevelVariantSetsEditor Delegate) override
+	{
+		OnLevelVariantSetsEditor = Delegate;
+	}
+
+	virtual void UnregisterOnLevelVariantSetsDelegate() override
+	{
+		OnLevelVariantSetsEditor.Unbind();
+	}
+
+	virtual FOnLevelVariantSetsEditor GetOnLevelVariantSetsEditorOpened() const override
+	{
+		return OnLevelVariantSetsEditor;
+	}
+
+    virtual UObject* CreateLevelVariantSetsAsset(const FString& AssetName, const FString& PackagePath, bool bForceOverwrite) override
+	{
+		IAssetTools& AssetTools = FModuleManager::GetModuleChecked<FAssetToolsModule>("AssetTools").Get();
+		UObject* NewAsset = nullptr;
+
+		FString SafePackagePath = UPackageTools::SanitizePackageName(PackagePath);
+		FString SafeAssetName = ObjectTools::SanitizeObjectName(AssetName);
+
+		if (SafePackagePath.IsEmpty() || SafeAssetName.IsEmpty())
+		{
+			UE_LOG(LogVariantManagerContentEditor, Error, TEXT("Invalid ULevelVariantSets asset name or package path!"));
+			return nullptr;
+		}
+
+		// Attempt to create a new asset by scanning all available factories until we find one that supports ULevelVariantSets
+		for (TObjectIterator<UClass> It ; It ; ++It)
+		{
+			UClass* CurrentClass = *It;
+			if (CurrentClass->IsChildOf(UFactory::StaticClass()) && !(CurrentClass->HasAnyClassFlags(CLASS_Abstract)))
+			{
+				UFactory* Factory = Cast<UFactory>(CurrentClass->GetDefaultObject());
+				if (Factory->CanCreateNew() && Factory->ImportPriority >= 0 && Factory->SupportedClass == ULevelVariantSets::StaticClass())
+				{
+					// Try deleting existing asset
+					// CreateAsset will always call CanCreateAsset, which will spawn a dialog in case of conflicts.
+					// We want to avoid the dialog, but also still use CreateAsset so that only VariantManagerFactoryNew is responsible for
+					// creating LVS assets. This means we have to manually (try to) delete the asset first
+					// Reference: UAssetToolsImpl::CanCreateAsset
+					if (bForceOverwrite)
+					{
+						const FString PackageName = SafePackagePath + TEXT("/") + SafeAssetName;
+						UPackage* Pkg = CreatePackage(nullptr,*PackageName);
+
+						// Search for UObject instead of ULevelVariantSets as we also want to catch redirectors
+						UObject* ExistingObject = StaticFindObject( UObject::StaticClass(), Pkg, *SafeAssetName );
+
+						if (ExistingObject != nullptr)
+						{
+							// Try to fixup a redirector before we delete it
+							if (ExistingObject->GetClass()->IsChildOf(UObjectRedirector::StaticClass()))
+							{
+								AssetTools.FixupReferencers({Cast<UObjectRedirector>(ExistingObject)});
+							}
+
+							bool bDeleteSucceeded = ObjectTools::DeleteSingleObject( ExistingObject );
+							if (bDeleteSucceeded)
+							{
+								// Force GC so we can cleanly create a new asset (and not do an 'in place' replacement)
+								CollectGarbage( GARBAGE_COLLECTION_KEEPFLAGS );
+
+								// Old package will be GC'ed... create a new one here
+								Pkg = CreatePackage(nullptr,*PackageName);
+								Pkg->MarkAsFullyLoaded();
+							}
+
+							if (!bDeleteSucceeded || !IsUniqueObjectName(*AssetName, Pkg))
+							{
+								UE_LOG(LogVariantManagerContentEditor, Error, TEXT("Failed to delete old ULevelVariantSets asset at '%s'. New asset can't be created."), *PackageName);
+								return nullptr;
+							}
+						}
+					}
+
+					NewAsset = AssetTools.CreateAsset(SafeAssetName, SafePackagePath, ULevelVariantSets::StaticClass(), Factory);
+					break;
+				}
+			}
+		}
+
+		return NewAsset;
+	}
+
+	virtual UObject* CreateLevelVariantSetsAssetWithDialog() override
+	{
+		IAssetTools& AssetTools = FModuleManager::GetModuleChecked<FAssetToolsModule>("AssetTools").Get();
+		UObject* NewAsset = nullptr;
+
+		// Attempt to create a new asset by scanning all available factories until we find one that supports ULevelVariantSets
+		for (TObjectIterator<UClass> It ; It ; ++It)
+		{
+			UClass* CurrentClass = *It;
+			if (CurrentClass->IsChildOf(UFactory::StaticClass()) && !(CurrentClass->HasAnyClassFlags(CLASS_Abstract)))
+			{
+				UFactory* Factory = Cast<UFactory>(CurrentClass->GetDefaultObject());
+				if (Factory->CanCreateNew() && Factory->ImportPriority >= 0 && Factory->SupportedClass == ULevelVariantSets::StaticClass())
+				{
+					NewAsset = AssetTools.CreateAssetWithDialog(ULevelVariantSets::StaticClass(), Factory);
+					break;
+				}
+			}
+		}
+
+		return NewAsset;
+	}
+
+	virtual AActor* GetOrCreateLevelVariantSetsActor(UObject* LevelVariantSetsAsset, bool bForceCreate) override
+	{
+		if (!bForceCreate)
+		{
+			// Check to see if we have a LVSA for this asset already
+			TArray< UObject* > WorldObjects;
+			GetObjectsWithOuter( GWorld, WorldObjects );
+
+			for( UObject* Object : WorldObjects )
+			{
+				if ( ALevelVariantSetsActor* Actor = Cast< ALevelVariantSetsActor >( Object ) )
+				{
+					if ( Actor->GetLevelVariantSets(true) == LevelVariantSetsAsset && !Actor->IsActorBeingDestroyed() )
+					{
+						return Actor;
+					}
+				}
+			}
+		}
+
+		UActorFactory* ActorFactory = GEditor->FindActorFactoryForActorClass(ALevelVariantSetsActor::StaticClass());
+		if (ActorFactory == nullptr)
+		{
+			UE_LOG(LogVariantManagerContentEditor, Error, TEXT("Failed to create LevelVariantSetsActor!"));
+			return nullptr;
+		}
+
+		return GEditor->UseActorFactory(ActorFactory, FAssetData(LevelVariantSetsAsset), &FTransform::Identity);
+	}
+
+private:
+	FOnLevelVariantSetsEditor OnLevelVariantSetsEditor;
+	TSharedPtr<IAssetTypeActions> AssetActions;
+};
+
+IMPLEMENT_MODULE(FVariantManagerContentEditorModule, VariantManagerContentEditor);
+
+#undef LOCTEXT_NAMESPACE
+
diff --git a/Engine/Plugins/Enterprise/VariantManagerContent/Source/VariantManagerContentEditor/Private/VariantManagerFactoryNew.cpp b/Engine/Plugins/Enterprise/VariantManagerContent/Source/VariantManagerContentEditor/Private/VariantManagerFactoryNew.cpp
new file mode 100644
--- /dev/null
+++ b/Engine/Plugins/Enterprise/VariantManagerContent/Source/VariantManagerContentEditor/Private/VariantManagerFactoryNew.cpp
@@ -0,0 +1,25 @@
+// Copyright 1998-2019 Epic Games, Inc. All Rights Reserved.
+
+#include "VariantManagerFactoryNew.h"
+
+#include "LevelVariantSets.h"
+
+UVariantManagerFactoryNew::UVariantManagerFactoryNew(const FObjectInitializer& ObjectInitializer)
+	: Super(ObjectInitializer)
+{
+	bCreateNew = true;
+	bEditAfterNew = true;
+	SupportedClass = ULevelVariantSets::StaticClass();
+}
+
+UObject* UVariantManagerFactoryNew::FactoryCreateNew(UClass* Class, UObject* InParent, FName Name, EObjectFlags Flags, UObject* Context, FFeedbackContext* Warn)
+{
+	ULevelVariantSets* VariantSets = NewObject<ULevelVariantSets>(InParent, Name, Flags|RF_Transactional);
+
+	return VariantSets;
+}
+
+bool UVariantManagerFactoryNew::ShouldShowInNewMenu() const
+{
+	return true;
+}
diff --git a/Engine/Plugins/Enterprise/VariantManagerContent/Source/VariantManagerContentEditor/Private/VariantManagerTestActor.cpp b/Engine/Plugins/Enterprise/VariantManagerContent/Source/VariantManagerContentEditor/Private/VariantManagerTestActor.cpp
new file mode 100644
--- /dev/null
+++ b/Engine/Plugins/Enterprise/VariantManagerContent/Source/VariantManagerContentEditor/Private/VariantManagerTestActor.cpp
@@ -0,0 +1,12 @@
+// Copyright 1998-2019 Epic Games, Inc. All Rights Reserved.
+
+#include "VariantManagerTestActor.h"
+
+#include "CoreMinimal.h"
+#include "UObject/ObjectMacros.h"
+
+
+AVariantManagerTestActor::AVariantManagerTestActor(const FObjectInitializer& Init)
+	: Super(Init)
+{
+}
\ No newline at end of file
diff --git a/Engine/Plugins/Enterprise/VariantManagerContent/Source/VariantManagerContentEditor/Public/LevelVariantSetsActorFactory.h b/Engine/Plugins/Enterprise/VariantManagerContent/Source/VariantManagerContentEditor/Public/LevelVariantSetsActorFactory.h
new file mode 100644
--- /dev/null
+++ b/Engine/Plugins/Enterprise/VariantManagerContent/Source/VariantManagerContentEditor/Public/LevelVariantSetsActorFactory.h
@@ -0,0 +1,28 @@
+// Copyright 1998-2019 Epic Games, Inc. All Rights Reserved.
+
+#pragma once
+
+#include "ActorFactories/ActorFactory.h"
+#include "CoreMinimal.h"
+#include "UObject/ObjectMacros.h"
+
+#include "LevelVariantSetsActorFactory.generated.h"
+
+class AActor;
+struct FAssetData;
+class ULevel;
+
+UCLASS()
+class ULevelVariantSetsActorFactory : public UActorFactory
+{
+	GENERATED_UCLASS_BODY()
+
+	// Begin UActorFactory Interface
+	virtual AActor* SpawnActor( UObject* Asset, ULevel* InLevel, const FTransform& Transform, EObjectFlags ObjectFlags, const FName Name ) override;
+	virtual bool CanCreateActorFrom( const FAssetData& AssetData, FText& OutErrorMsg ) override;
+	virtual UObject* GetAssetFromActorInstance(AActor* ActorInstance) override;
+	// End UActorFactory Interface
+};
+
+
+
diff --git a/Engine/Plugins/Enterprise/VariantManagerContent/Source/VariantManagerContentEditor/Public/VariantManagerContentEditorModule.h b/Engine/Plugins/Enterprise/VariantManagerContent/Source/VariantManagerContentEditor/Public/VariantManagerContentEditorModule.h
new file mode 100644
--- /dev/null
+++ b/Engine/Plugins/Enterprise/VariantManagerContent/Source/VariantManagerContentEditor/Public/VariantManagerContentEditorModule.h
@@ -0,0 +1,40 @@
+// Copyright 1998-2019 Epic Games, Inc. All Rights Reserved.
+
+#pragma once
+
+#include "Modules/ModuleInterface.h"
+#include "Modules/ModuleManager.h"		// For inline LoadModuleChecked()
+#include "Toolkits/IToolkit.h"
+
+#define VARIANTMANAGERCONTENTEDITORMODULE_MODULE_NAME TEXT("VariantManagerContentEditor")
+
+class AActor;
+class ULevelVariantSets;
+
+DECLARE_DELEGATE_ThreeParams( FOnLevelVariantSetsEditor, const EToolkitMode::Type, const TSharedPtr< class IToolkitHost >&, class ULevelVariantSets*);
+
+class IVariantManagerContentEditorModule : public IModuleInterface
+{
+public:
+	static inline IVariantManagerContentEditorModule& Get()
+	{
+		return FModuleManager::LoadModuleChecked<IVariantManagerContentEditorModule>(VARIANTMANAGERCONTENTEDITORMODULE_MODULE_NAME);
+	}
+
+	static inline bool IsAvailable()
+	{
+		return FModuleManager::Get().IsModuleLoaded(VARIANTMANAGERCONTENTEDITORMODULE_MODULE_NAME);
+	}
+
+	// This will be executed whenever FLevelVariantSetsAssetActions::OpenAssetEditor is called with a valid asset.
+	// This so that the editor toolkit can be on the VariantManager plugin and the AssetActions can be on this module
+	virtual void RegisterOnLevelVariantSetsDelegate(FOnLevelVariantSetsEditor Delegate) = 0;
+	virtual void UnregisterOnLevelVariantSetsDelegate() = 0;
+	virtual FOnLevelVariantSetsEditor GetOnLevelVariantSetsEditorOpened() const = 0;
+
+	virtual UObject* CreateLevelVariantSetsAssetWithDialog() = 0;
+	virtual UObject* CreateLevelVariantSetsAsset(const FString& AssetName, const FString& PackagePath, bool bForceOverwrite = false) = 0;
+	virtual AActor* GetOrCreateLevelVariantSetsActor(UObject* LevelVariantSetsAsset, bool bForceCreate=false) = 0;
+};
+
+
diff --git a/Engine/Plugins/Enterprise/VariantManagerContent/Source/VariantManagerContentEditor/Public/VariantManagerFactoryNew.h b/Engine/Plugins/Enterprise/VariantManagerContent/Source/VariantManagerContentEditor/Public/VariantManagerFactoryNew.h
new file mode 100644
--- /dev/null
+++ b/Engine/Plugins/Enterprise/VariantManagerContent/Source/VariantManagerContentEditor/Public/VariantManagerFactoryNew.h
@@ -0,0 +1,21 @@
+// Copyright 1998-2019 Epic Games, Inc. All Rights Reserved.
+
+#pragma once
+
+#include "CoreMinimal.h"
+#include "Factories/Factory.h"
+#include "UObject/ObjectMacros.h"
+
+#include "VariantManagerFactoryNew.generated.h"
+
+UCLASS(hidecategories=Object)
+class UVariantManagerFactoryNew : public UFactory
+{
+	GENERATED_UCLASS_BODY()
+
+public:
+	// UFactory Interface
+	virtual UObject* FactoryCreateNew(UClass* Class, UObject* InParent, FName Name, EObjectFlags Flags, UObject* Context, FFeedbackContext* Warn) override;
+	virtual bool ShouldShowInNewMenu() const override;
+    //~ End UFactory Interface
+};
diff --git a/Engine/Plugins/Enterprise/VariantManagerContent/Source/VariantManagerContentEditor/Public/VariantManagerTestActor.h b/Engine/Plugins/Enterprise/VariantManagerContent/Source/VariantManagerContentEditor/Public/VariantManagerTestActor.h
new file mode 100644
--- /dev/null
+++ b/Engine/Plugins/Enterprise/VariantManagerContent/Source/VariantManagerContentEditor/Public/VariantManagerTestActor.h
@@ -0,0 +1,112 @@
+// Copyright 1998-2019 Epic Games, Inc. All Rights Reserved.
+
+#pragma once
+
+#include "CoreMinimal.h"
+#include "GameFramework/Actor.h"
+#include "UObject/ObjectMacros.h"
+
+#include "VariantManagerTestActor.generated.h"
+
+UENUM(BlueprintType)
+enum class EVariantManagerTestEnum : uint8
+{
+	None = 0 UMETA(Hidden),
+	FirstOption = 1,
+	SecondOption = 3,
+	ThirdOption = 45,
+};
+
+UCLASS(hideCategories=(Rendering, Physics, LOD, Activation, Input, Actor, Cooking))
+class VARIANTMANAGERCONTENTEDITOR_API AVariantManagerTestActor : public AActor
+{
+public:
+
+	GENERATED_BODY()
+
+	AVariantManagerTestActor(const FObjectInitializer& Init);
+
+	UPROPERTY(EditAnywhere, BlueprintReadWrite, Category="Template", meta=(ToolTip="Captured byte property"))
+	EVariantManagerTestEnum EnumWithNoDefault;
+
+	UPROPERTY(EditAnywhere, BlueprintReadWrite, Category="Template", meta=(ToolTip="Captured byte property"))
+	EVariantManagerTestEnum EnumWithSecondDefault = EVariantManagerTestEnum::SecondOption;
+
+	UPROPERTY(EditAnywhere, BlueprintReadWrite, Category="Template", meta=(ToolTip="Captured byte property"))
+	uint8 CapturedByteProperty;
+
+	//UPROPERTY(EditAnywhere, BlueprintReadWrite, Category="Template", meta=(ToolTip="Captured uint16 property"))
+	//uint16 CapturedUInt16Property;
+
+	//UPROPERTY(EditAnywhere, BlueprintReadWrite, Category="Template", meta=(ToolTip="Captured uint32 property"))
+	//uint32 CapturedUInt32Property;
+
+	//UPROPERTY(EditAnywhere, BlueprintReadWrite, Category="Template", meta=(ToolTip="Captured uint16 property"))
+	//uint64 CapturedUInt64Property;
+
+	//UPROPERTY(EditAnywhere, BlueprintReadWrite, Category="Template", meta=(ToolTip="Captured int8 property"))
+	//int8 CapturedInt8Property;
+
+	//UPROPERTY(EditAnywhere, BlueprintReadWrite, Category="Template", meta=(ToolTip="Captured int16 property"))
+	//int16 CapturedInt16Property;
+
+	UPROPERTY(EditAnywhere, BlueprintReadWrite, Category="Template", meta=(ToolTip="Captured int32 property"))
+	int32 CapturedIntProperty = 7;
+
+	//UPROPERTY(EditAnywhere, BlueprintReadWrite, Category="Template", meta=(ToolTip="Captured int64 property"))
+	//int64 CapturedInt64Property;
+
+	UPROPERTY(EditAnywhere, BlueprintReadWrite, Category="Template", meta=(ToolTip="Captured float property"))
+	float CapturedFloatProperty = 2.3f;
+
+	//UPROPERTY(EditAnywhere, BlueprintReadWrite, Category="Template", meta=(ToolTip="Captured double property"))
+	//double CapturedDoubleProperty;
+
+	UPROPERTY(EditAnywhere, BlueprintReadWrite, Category="Template", meta=(ToolTip="Captured boolean property"))
+	bool bCapturedBoolProperty;
+
+	UPROPERTY(EditAnywhere, BlueprintReadWrite, Category="Template", meta=(ToolTip="Captured UObject property"))
+	UObject* CapturedObjectProperty;
+
+	UPROPERTY(EditAnywhere, BlueprintReadWrite, Category="Template", meta=(ToolTip="Captured UInterface property"))
+	FScriptInterface CapturedInterfaceProperty;
+
+	UPROPERTY(EditAnywhere, BlueprintReadWrite, Category="Template", meta=(ToolTip="Captured FName property"))
+	FName CapturedNameProperty;
+
+	UPROPERTY(EditAnywhere, BlueprintReadWrite, Category="Template", meta=(ToolTip="Captured FString property"))
+	FString CapturedStrProperty;
+
+	UPROPERTY(EditAnywhere, BlueprintReadWrite, Category="Template", meta=(ToolTip="Captured FText property"))
+	FText CapturedTextProperty;
+
+	UPROPERTY(EditAnywhere, BlueprintReadWrite, Category="Template", meta=(ToolTip="Captured FRotator property"))
+	FRotator CapturedRotatorProperty;
+
+	UPROPERTY(EditAnywhere, BlueprintReadWrite, Category="Template", meta=(ToolTip="Captured FColor property"))
+	FColor CapturedColorProperty;
+
+	UPROPERTY(EditAnywhere, BlueprintReadWrite, Category="Template", meta=(ToolTip="Captured FLinearColor property"))
+	FLinearColor CapturedLinearColorProperty;
+
+	UPROPERTY(EditAnywhere, BlueprintReadWrite, Category="Template", meta=(ToolTip="Captured FVector property"))
+	FVector CapturedVectorProperty;
+
+	UPROPERTY(EditAnywhere, BlueprintReadWrite, Category="Template", meta=(ToolTip="Captured FQuat property"))
+	FQuat CapturedQuatProperty;
+
+	UPROPERTY(EditAnywhere, BlueprintReadWrite, Category="Template", meta=(ToolTip="Captured FVector4 property"))
+	FVector4 CapturedVector4Property;
+
+	UPROPERTY(EditAnywhere, BlueprintReadWrite, Category="Template", meta=(ToolTip="Captured FVector2D property"))
+	FVector2D CapturedVector2DProperty;
+
+	UPROPERTY(EditAnywhere, BlueprintReadWrite, Category="Template", meta=(ToolTip="Captured FIntPoint property"))
+	FIntPoint CapturedIntPointProperty;
+
+	UPROPERTY(EditAnywhere, BlueprintReadWrite, Category="Template", meta=(ToolTip="Captured UObject array property"))
+	TArray<UObject*> CapturedUObjectArrayProperty;
+
+	UPROPERTY(EditAnywhere, BlueprintReadWrite, Category="Template", meta=(ToolTip="Captured FVector array property"))
+	TArray<FVector> CapturedVectorArrayProperty;
+};
diff --git a/Engine/Plugins/Enterprise/VariantManagerContent/Source/VariantManagerContentEditor/VariantManagerContentEditor.Build.cs b/Engine/Plugins/Enterprise/VariantManagerContent/Source/VariantManagerContentEditor/VariantManagerContentEditor.Build.cs
new file mode 100644
--- /dev/null
+++ b/Engine/Plugins/Enterprise/VariantManagerContent/Source/VariantManagerContentEditor/VariantManagerContentEditor.Build.cs
@@ -0,0 +1,34 @@
+// Copyright 1998-2019 Epic Games, Inc. All Rights Reserved.
+
+namespace UnrealBuildTool.Rules
+{
+	public class VariantManagerContentEditor : ModuleRules
+	{
+		public VariantManagerContentEditor(ReadOnlyTargetRules Target)
+			: base(Target)
+		{
+			PublicDependencyModuleNames.AddRange(
+				new string[]
+				{
+					"Core",
+					"CoreUObject",
+					"Engine",
+				}
+			);
+
+			PrivateDependencyModuleNames.AddRange(
+				new string[]
+				{
+					"EditorStyle",
+					"InputCore", // For ListView keyboard control
+					"Slate",
+					"SlateCore",
+					"ToolMenus",
+					"UnrealEd",
+					"VariantManagerContent",
+					"WorkspaceMenuStructure",
+				}
+			);
+		}
+	}
+}
\ No newline at end of file
diff --git a/Engine/Plugins/Enterprise/VariantManagerContent/VariantManagerContent.uplugin b/Engine/Plugins/Enterprise/VariantManagerContent/VariantManagerContent.uplugin
--- a/Engine/Plugins/Enterprise/VariantManagerContent/VariantManagerContent.uplugin
+++ b/Engine/Plugins/Enterprise/VariantManagerContent/VariantManagerContent.uplugin
@@ -17,9 +17,14 @@
 	"Modules" :
 	[
 		{
-			"Name" : "VariantManagerContent",
-			"Type" : "Runtime",
-			"LoadingPhase" : "Default"
+			"Name": "VariantManagerContent",
+			"Type": "Runtime",
+			"LoadingPhase": "PostConfigInit"
+		},
+		{
+			"Name": "VariantManagerContentEditor",
+			"Type": "Editor",
+			"LoadingPhase": "Default"
 		}
 	]
 }
\ No newline at end of file
diff --git a/Engine/Source/Developer/Datasmith/DatasmithCore/Private/DatasmithVariantElementsImpl.h b/Engine/Source/Developer/Datasmith/DatasmithCore/Private/DatasmithVariantElementsImpl.h
--- a/Engine/Source/Developer/Datasmith/DatasmithCore/Private/DatasmithVariantElementsImpl.h
+++ b/Engine/Source/Developer/Datasmith/DatasmithCore/Private/DatasmithVariantElementsImpl.h
@@ -1,8 +1,9 @@
 // Copyright 1998-2019 Epic Games, Inc. All Rights Reserved.
 #pragma once
 
-#include "DatasmithVariantElements.h"
+#include "DatasmithDefinitions.h"
 #include "DatasmithSceneElementsImpl.h"
+#include "DatasmithVariantElements.h"
 
 #include "Containers/Array.h"
 #include "Templates/SharedPointer.h"
@@ -13,6 +14,7 @@ class FDatasmithBasePropertyCaptureElementImpl : public FDatasmithElementImpl< I
 public:
 	explicit FDatasmithBasePropertyCaptureElementImpl(EDatasmithElementVariantSubType InSubType = EDatasmithElementVariantSubType::PropertyCapture)
 		: FDatasmithElementImpl< InterfaceType >(nullptr, EDatasmithElementType::Variant, (uint64)InSubType)
+		, Category(EDatasmithPropertyCategory::Undefined)
 	{
 	}
 
@@ -26,19 +28,19 @@ public:
 		return PropertyPath;
 	}
 
-	virtual void SetCategory(EPropertyValueCategory InCategory) override
+	virtual void SetCategory(EDatasmithPropertyCategory InCategory) override
 	{
 		Category = InCategory;
 	}
 
-	virtual EPropertyValueCategory GetCategory() const override
+	virtual EDatasmithPropertyCategory GetCategory() const override
 	{
 		return Category;
 	}
 
 private:
 	FString PropertyPath;
-	EPropertyValueCategory Category;
+	EDatasmithPropertyCategory Category;
 };
 
 class FDatasmithPropertyCaptureElementImpl : public FDatasmithBasePropertyCaptureElementImpl< IDatasmithPropertyCaptureElement >
diff --git a/Engine/Source/Developer/Datasmith/DatasmithCore/Public/DatasmithDefinitions.h b/Engine/Source/Developer/Datasmith/DatasmithCore/Public/DatasmithDefinitions.h
--- a/Engine/Source/Developer/Datasmith/DatasmithCore/Public/DatasmithDefinitions.h
+++ b/Engine/Source/Developer/Datasmith/DatasmithCore/Public/DatasmithDefinitions.h
@@ -335,6 +335,26 @@ enum class EDatasmithCompletionMode : uint8
 	ProjectDefault,
 };
 
+/**
+ * Describes a category of an UPropertyValue asset, indicating types of
+ * properties that require special handling for any reason.
+ * Mirrors EPropertyValueCategory
+ */
+UENUM()
+enum class EDatasmithPropertyCategory : uint8
+{
+	Undefined = 0,
+	Generic = 1,
+	RelativeLocation = 2,
+	RelativeRotation = 4,
+	RelativeScale3D = 8,
+	Visibility = 16,
+	Material = 32,
+	Color = 64,
+	Option = 128
+};
+ENUM_CLASS_FLAGS(EDatasmithPropertyCategory)
+
 static const TCHAR* KeyValuePropertyTypeStrings[] = { TEXT("String"), TEXT("Color"), TEXT("Float"), TEXT("Bool"), TEXT("Texture"), TEXT("Vector") };
 
 // HOST NAME
diff --git a/Engine/Source/Developer/Datasmith/DatasmithCore/Public/DatasmithVariantElements.h b/Engine/Source/Developer/Datasmith/DatasmithCore/Public/DatasmithVariantElements.h
--- a/Engine/Source/Developer/Datasmith/DatasmithCore/Public/DatasmithVariantElements.h
+++ b/Engine/Source/Developer/Datasmith/DatasmithCore/Public/DatasmithVariantElements.h
@@ -2,20 +2,19 @@
 
 #pragma once
 
+#include "DatasmithDefinitions.h"
 #include "IDatasmithSceneElements.h"
 
 #include "Templates/SharedPointer.h"
 
-enum class EPropertyValueCategory : uint8;
-
 class DATASMITHCORE_API IDatasmithBasePropertyCaptureElement : public IDatasmithElement
 {
 public:
 	virtual void SetPropertyPath(const FString& Path) = 0;
 	virtual const FString& GetPropertyPath() const = 0;
 
-	virtual void SetCategory(EPropertyValueCategory Category) = 0;
-	virtual EPropertyValueCategory GetCategory() const = 0;
+	virtual void SetCategory(EDatasmithPropertyCategory Category) = 0;
+	virtual EDatasmithPropertyCategory GetCategory() const = 0;
 };
 
 class DATASMITHCORE_API IDatasmithPropertyCaptureElement : public IDatasmithBasePropertyCaptureElement
