More CIS Deprecation fixes.

diff --git a/Engine/Plugins/Enterprise/VariantManager/Source/VariantManager/Private/SVariantManager.cpp b/Engine/Plugins/Enterprise/VariantManager/Source/VariantManager/Private/SVariantManager.cpp
--- a/Engine/Plugins/Enterprise/VariantManager/Source/VariantManager/Private/SVariantManager.cpp
+++ b/Engine/Plugins/Enterprise/VariantManager/Source/VariantManager/Private/SVariantManager.cpp
@@ -2384,7 +2384,7 @@ void SVariantManager::OnObjectTransacted(UObject* Object, const class FTransacti
 		FString ParentActorPath;
 		if (ASwitchActor* SwitchActorParent = Cast<ASwitchActor>(TargetActor->GetAttachParentActor()))
 		{
-			if (Object->IsA(USceneComponent::StaticClass()) && Event.GetChangedProperties().Contains(GET_MEMBER_NAME_CHECKED(USceneComponent, bVisible)))
+			if (Object->IsA(USceneComponent::StaticClass()) && Event.GetChangedProperties().Contains(USceneComponent::GetVisiblePropertyName()))
 			{
 				ParentActorPath = SwitchActorParent->GetPathName();
 				bool bSwitchWasCapturedAlready = CachedDisplayedActorPaths.Contains(ParentActorPath);
diff --git a/Engine/Plugins/Enterprise/VariantManager/Source/VariantManager/Private/VariantManagerUtils.cpp b/Engine/Plugins/Enterprise/VariantManager/Source/VariantManager/Private/VariantManagerUtils.cpp
--- a/Engine/Plugins/Enterprise/VariantManager/Source/VariantManager/Private/VariantManagerUtils.cpp
+++ b/Engine/Plugins/Enterprise/VariantManager/Source/VariantManager/Private/VariantManagerUtils.cpp
@@ -66,7 +66,7 @@ UStructProperty* FVariantManagerUtils::GetRelativeLocationProperty()
 {
 	if (!RelativeLocationProperty)
 	{
-		RelativeLocationProperty = FindField<UStructProperty>( USceneComponent::StaticClass(), GET_MEMBER_NAME_CHECKED( USceneComponent, RelativeLocation ) );
+		RelativeLocationProperty = FindField<UStructProperty>(USceneComponent::StaticClass(), *USceneComponent::GetRelativeLocationPropertyName().ToString());
 	}
 
 	return RelativeLocationProperty;
@@ -76,7 +76,7 @@ UStructProperty* FVariantManagerUtils::GetRelativeRotationProperty()
 {
 	if (!RelativeRotationProperty)
 	{
-		RelativeRotationProperty = FindField<UStructProperty>( USceneComponent::StaticClass(), GET_MEMBER_NAME_CHECKED( USceneComponent, RelativeRotation ) );
+		RelativeRotationProperty = FindField<UStructProperty>(USceneComponent::StaticClass(), *USceneComponent::GetRelativeRotationPropertyName().ToString());
 	}
 
 	return RelativeRotationProperty;
@@ -86,7 +86,7 @@ UStructProperty* FVariantManagerUtils::GetRelativeScale3DProperty()
 {
 	if (!RelativeScale3DProperty)
 	{
-		RelativeScale3DProperty = FindField<UStructProperty>( USceneComponent::StaticClass(), GET_MEMBER_NAME_CHECKED( USceneComponent, RelativeScale3D ) );
+		RelativeScale3DProperty = FindField<UStructProperty>(USceneComponent::StaticClass(), *USceneComponent::GetRelativeScale3DPropertyName().ToString());
 	}
 
 	return RelativeScale3DProperty;
@@ -96,7 +96,7 @@ UBoolProperty* FVariantManagerUtils::GetVisibilityProperty()
 {
 	if (!VisiblityProperty)
 	{
-		VisiblityProperty = FindField<UBoolProperty>( USceneComponent::StaticClass(), GET_MEMBER_NAME_CHECKED( USceneComponent, bVisible ) );
+		VisiblityProperty = FindField<UBoolProperty>(USceneComponent::StaticClass(), *USceneComponent::GetVisiblePropertyName().ToString());
 	}
 
 	return VisiblityProperty;
diff --git a/Engine/Source/Editor/UnrealEd/Private/LevelEditorViewport.cpp b/Engine/Source/Editor/UnrealEd/Private/LevelEditorViewport.cpp
--- a/Engine/Source/Editor/UnrealEd/Private/LevelEditorViewport.cpp
+++ b/Engine/Source/Editor/UnrealEd/Private/LevelEditorViewport.cpp
@@ -3462,7 +3462,7 @@ void FLevelEditorViewportClient::MoveLockedActorToCamera()
 			TOptional<FRotator> PreviousRotator;
 			if (ActiveActorLockComponent)
 			{
-				PreviousRotator = ActiveActorLockComponent->RelativeRotation;
+				PreviousRotator = ActiveActorLockComponent->GetRelativeRotation();
 			}
 
 			// If we're locked to a camera then we're reflecting the camera view and not the actor position. We need to reflect that delta when we reposition the piloted actor
@@ -3488,7 +3488,7 @@ void FLevelEditorViewportClient::MoveLockedActorToCamera()
 				FRotator ActorRotWind, ActorRotRem;
 				Rot.GetWindingAndRemainder(ActorRotWind, ActorRotRem);
 				const FQuat ActorQ = ActorRotRem.Quaternion();
-				const FQuat ResultQ = ActiveActorLockComponent->RelativeRotation.Quaternion();
+				const FQuat ResultQ = ActiveActorLockComponent->GetRelativeRotation().Quaternion();
 				FRotator NewActorRotRem = FRotator(ResultQ);
 				ActorRotRem.SetClosestToMe(NewActorRotRem);
 				FRotator DeltaRot = NewActorRotRem - ActorRotRem;
