CIS: Fixes deprecation warnings in latest merge from main.

#rnx

diff --git a/Engine/Plugins/Enterprise/VariantManager/Source/VariantManager/Private/SVariantManager.cpp b/Engine/Plugins/Enterprise/VariantManager/Source/VariantManager/Private/SVariantManager.cpp
--- a/Engine/Plugins/Enterprise/VariantManager/Source/VariantManager/Private/SVariantManager.cpp
+++ b/Engine/Plugins/Enterprise/VariantManager/Source/VariantManager/Private/SVariantManager.cpp
@@ -2384,7 +2384,7 @@ void SVariantManager::OnObjectTransacted(UObject* Object, const class FTransacti
 		FString ParentActorPath;
 		if (ASwitchActor* SwitchActorParent = Cast<ASwitchActor>(TargetActor->GetAttachParentActor()))
 		{
-			if (Object->IsA(USceneComponent::StaticClass()) && Event.GetChangedProperties().Contains(GET_MEMBER_NAME_CHECKED(USceneComponent, bVisible)))
+			if (Object->IsA(USceneComponent::StaticClass()) && Event.GetChangedProperties().Contains(GET_MEMBER_NAME_CHECKED(USceneComponent, USceneComponent::GetVisiblePropertyName())))
 			{
 				ParentActorPath = SwitchActorParent->GetPathName();
 				bool bSwitchWasCapturedAlready = CachedDisplayedActorPaths.Contains(ParentActorPath);
diff --git a/Engine/Plugins/Enterprise/VariantManager/Source/VariantManager/Private/VariantManagerUtils.cpp b/Engine/Plugins/Enterprise/VariantManager/Source/VariantManager/Private/VariantManagerUtils.cpp
--- a/Engine/Plugins/Enterprise/VariantManager/Source/VariantManager/Private/VariantManagerUtils.cpp
+++ b/Engine/Plugins/Enterprise/VariantManager/Source/VariantManager/Private/VariantManagerUtils.cpp
@@ -66,7 +66,7 @@ UStructProperty* FVariantManagerUtils::GetRelativeLocationProperty()
 {
 	if (!RelativeLocationProperty)
 	{
-		RelativeLocationProperty = FindField<UStructProperty>( USceneComponent::StaticClass(), GET_MEMBER_NAME_CHECKED( USceneComponent, RelativeLocation ) );
+		RelativeLocationProperty = FindField<UStructProperty>( USceneComponent::StaticClass(), GET_MEMBER_NAME_CHECKED( USceneComponent, USceneComponent::GetRelativeLocationPropertyName() ) );
 	}
 
 	return RelativeLocationProperty;
@@ -76,7 +76,7 @@ UStructProperty* FVariantManagerUtils::GetRelativeRotationProperty()
 {
 	if (!RelativeRotationProperty)
 	{
-		RelativeRotationProperty = FindField<UStructProperty>( USceneComponent::StaticClass(), GET_MEMBER_NAME_CHECKED( USceneComponent, RelativeRotation ) );
+		RelativeRotationProperty = FindField<UStructProperty>( USceneComponent::StaticClass(), GET_MEMBER_NAME_CHECKED( USceneComponent, USceneComponent::GetRelativeRotationPropertyName() ) );
 	}
 
 	return RelativeRotationProperty;
@@ -86,7 +86,7 @@ UStructProperty* FVariantManagerUtils::GetRelativeScale3DProperty()
 {
 	if (!RelativeScale3DProperty)
 	{
-		RelativeScale3DProperty = FindField<UStructProperty>( USceneComponent::StaticClass(), GET_MEMBER_NAME_CHECKED( USceneComponent, RelativeScale3D ) );
+		RelativeScale3DProperty = FindField<UStructProperty>( USceneComponent::StaticClass(), GET_MEMBER_NAME_CHECKED( USceneComponent, USceneComponent::GetRelativeScale3DPropertyName() ) );
 	}
 
 	return RelativeScale3DProperty;
@@ -96,7 +96,7 @@ UBoolProperty* FVariantManagerUtils::GetVisibilityProperty()
 {
 	if (!VisiblityProperty)
 	{
-		VisiblityProperty = FindField<UBoolProperty>( USceneComponent::StaticClass(), GET_MEMBER_NAME_CHECKED( USceneComponent, bVisible ) );
+		VisiblityProperty = FindField<UBoolProperty>( USceneComponent::StaticClass(), GET_MEMBER_NAME_CHECKED( USceneComponent, USceneComponent::GetVisiblePropertyName() ) );
 	}
 
 	return VisiblityProperty;
diff --git a/Engine/Plugins/Experimental/DataPrepEditor/Source/DataPrepEditor/Private/DataPrepEditor.cpp b/Engine/Plugins/Experimental/DataPrepEditor/Source/DataPrepEditor/Private/DataPrepEditor.cpp
--- a/Engine/Plugins/Experimental/DataPrepEditor/Source/DataPrepEditor/Private/DataPrepEditor.cpp
+++ b/Engine/Plugins/Experimental/DataPrepEditor/Source/DataPrepEditor/Private/DataPrepEditor.cpp
@@ -1145,7 +1145,7 @@ bool FDataprepEditor::OnRequestClose()
 		const FText Title( LOCTEXT( "DataprepEditor_ProceedWithClose", "Proceed with close") );
 		const FText Message( LOCTEXT( "DataprepEditor_ConfirmClose", "Imported data was not committed! Closing the editor will discard imported data.\nDo you want to close anyway?" ) );
 
-		return ( OpenMsgDlgInt( EAppMsgType::YesNo, Message, Title ) == EAppReturnType::Yes );
+		return ( FMessageDialog::Open(EAppMsgType::YesNo, Message, &Title) == EAppReturnType::Yes );
 	}
 	return !bIgnoreCloseRequest;
 }
diff --git a/Engine/Plugins/Runtime/AR/Microsoft/HoloLensAR/Source/HoloLensAR/Private/HoloLensARSystem.cpp b/Engine/Plugins/Runtime/AR/Microsoft/HoloLensAR/Source/HoloLensAR/Private/HoloLensARSystem.cpp
--- a/Engine/Plugins/Runtime/AR/Microsoft/HoloLensAR/Source/HoloLensAR/Private/HoloLensARSystem.cpp
+++ b/Engine/Plugins/Runtime/AR/Microsoft/HoloLensAR/Source/HoloLensAR/Private/HoloLensARSystem.cpp
@@ -117,7 +117,7 @@ void FHoloLensARSystem::SetInterop(WindowsMixedReality::MixedRealityInterop* InW
 			OnStopARSession();
 
 #if WITH_EDITOR
-			if (GEditor && GEditor->bUseVRPreviewForPlayWorld)
+			if (GEditor && GEditor->IsVRPreviewActive())
 			{
 				GEditor->RequestEndPlayMap();
 			}
diff --git a/Engine/Plugins/Runtime/nDisplay/Source/DisplayClusterEditor/Private/DisplayClusterEditorEngine.cpp b/Engine/Plugins/Runtime/nDisplay/Source/DisplayClusterEditor/Private/DisplayClusterEditorEngine.cpp
--- a/Engine/Plugins/Runtime/nDisplay/Source/DisplayClusterEditor/Private/DisplayClusterEditorEngine.cpp
+++ b/Engine/Plugins/Runtime/nDisplay/Source/DisplayClusterEditor/Private/DisplayClusterEditorEngine.cpp
@@ -45,13 +45,15 @@ void UDisplayClusterEditorEngine::PreExit()
 	Super::PreExit();
 }
 
-void UDisplayClusterEditorEngine::PlayInEditor(UWorld* InWorld, bool bInSimulateInEditor, FPlayInEditorOverrides Overrides)
+void UDisplayClusterEditorEngine::StartPlayInEditorSession(FRequestPlaySessionParams& InRequestParams)
 {
-	UE_LOG(LogDisplayClusterEditorEngine, VeryVerbose, TEXT("UDisplayClusterEditorEngine::PlayInEditor"));
+	UE_LOG(LogDisplayClusterEditorEngine, VeryVerbose, TEXT("UDisplayClusterEditorEngine::StartPlayInEditorSession"));
+
+	UWorld* EditorWorld = GetEditorWorldContext().World();
 
 	// Find nDisplay root
 	UDisplayClusterRootComponent* RootComponent = nullptr;
-	for (AActor* Actor : InWorld->PersistentLevel->Actors)
+	for (AActor* Actor : EditorWorld->PersistentLevel->Actors)
 	{
 		if (Actor && !Actor->IsPendingKill())
 		{
@@ -84,11 +86,11 @@ void UDisplayClusterEditorEngine::PlayInEditor(UWorld* InWorld, bool bInSimulate
 				return;
 			}
 
-			DisplayClusterModule->StartScene(InWorld);
+			DisplayClusterModule->StartScene(EditorWorld);
 		}
 	}
 
-	Super::PlayInEditor(InWorld, bInSimulateInEditor, Overrides);
+	Super::StartPlayInEditorSession(InRequestParams);
 }
 
 void UDisplayClusterEditorEngine::Tick(float DeltaSeconds, bool bIdleMode)
diff --git a/Engine/Plugins/Runtime/nDisplay/Source/DisplayClusterEditor/Private/DisplayClusterEditorEngine.h b/Engine/Plugins/Runtime/nDisplay/Source/DisplayClusterEditor/Private/DisplayClusterEditorEngine.h
--- a/Engine/Plugins/Runtime/nDisplay/Source/DisplayClusterEditor/Private/DisplayClusterEditorEngine.h
+++ b/Engine/Plugins/Runtime/nDisplay/Source/DisplayClusterEditor/Private/DisplayClusterEditorEngine.h
@@ -21,7 +21,7 @@ class UDisplayClusterEditorEngine
 public:
 	virtual void Init(IEngineLoop* InEngineLoop) override;
 	virtual void PreExit() override;
-	virtual void PlayInEditor(UWorld* InWorld, bool bInSimulateInEditor, FPlayInEditorOverrides Overrides = FPlayInEditorOverrides()) override;
+	virtual void StartPlayInEditorSession(FRequestPlaySessionParams& InRequestParams) override;
 	virtual void Tick(float DeltaSeconds, bool bIdleMode) override;
 
 private:
diff --git a/Engine/Source/Editor/UnrealEd/Private/LevelEditorViewport.cpp b/Engine/Source/Editor/UnrealEd/Private/LevelEditorViewport.cpp
--- a/Engine/Source/Editor/UnrealEd/Private/LevelEditorViewport.cpp
+++ b/Engine/Source/Editor/UnrealEd/Private/LevelEditorViewport.cpp
@@ -3462,7 +3462,7 @@ void FLevelEditorViewportClient::MoveLockedActorToCamera()
 			TOptional<FRotator> PreviousRotator;
 			if (ActiveActorLockComponent)
 			{
-				PreviousRotator = ActiveActorLockComponent->RelativeRotation;
+				PreviousRotator = ActiveActorLockComponent->GetRelativeRotation();
 			}
 
 			// If we're locked to a camera then we're reflecting the camera view and not the actor position. We need to reflect that delta when we reposition the piloted actor
@@ -3488,7 +3488,7 @@ void FLevelEditorViewportClient::MoveLockedActorToCamera()
 				FRotator ActorRotWind, ActorRotRem;
 				Rot.GetWindingAndRemainder(ActorRotWind, ActorRotRem);
 				const FQuat ActorQ = ActorRotRem.Quaternion();
-				const FQuat ResultQ = ActiveActorLockComponent->RelativeRotation.Quaternion();
+				const FQuat ResultQ = ActiveActorLockComponent->GetRelativeRotation().Quaternion();
 				FRotator NewActorRotRem = FRotator(ResultQ);
 				ActorRotRem.SetClosestToMe(NewActorRotRem);
 				FRotator DeltaRot = NewActorRotRem - ActorRotRem;
