Use GET_MEMBER_NAME_CHECKED instead of converting to strings in VariantManagerUtils.

#rnx

diff --git a/Engine/Plugins/Enterprise/VariantManager/Source/VariantManager/Private/VariantManagerUtils.cpp b/Engine/Plugins/Enterprise/VariantManager/Source/VariantManager/Private/VariantManagerUtils.cpp
--- a/Engine/Plugins/Enterprise/VariantManager/Source/VariantManager/Private/VariantManagerUtils.cpp
+++ b/Engine/Plugins/Enterprise/VariantManager/Source/VariantManager/Private/VariantManagerUtils.cpp
@@ -66,7 +66,7 @@ UStructProperty* FVariantManagerUtils::GetRelativeLocationProperty()
 {
 	if (!RelativeLocationProperty)
 	{
-		RelativeLocationProperty = FindField<UStructProperty>(USceneComponent::StaticClass(), *USceneComponent::GetRelativeLocationPropertyName().ToString());
+		RelativeLocationProperty = FindField<UStructProperty>( USceneComponent::StaticClass(), GET_MEMBER_NAME_CHECKED( USceneComponent, USceneComponent::GetRelativeLocationPropertyName() ) );
 	}
 
 	return RelativeLocationProperty;
@@ -76,7 +76,7 @@ UStructProperty* FVariantManagerUtils::GetRelativeRotationProperty()
 {
 	if (!RelativeRotationProperty)
 	{
-		RelativeRotationProperty = FindField<UStructProperty>(USceneComponent::StaticClass(), *USceneComponent::GetRelativeRotationPropertyName().ToString());
+		RelativeRotationProperty = FindField<UStructProperty>( USceneComponent::StaticClass(), GET_MEMBER_NAME_CHECKED( USceneComponent, USceneComponent::GetRelativeRotationPropertyName() ) );
 	}
 
 	return RelativeRotationProperty;
@@ -86,7 +86,7 @@ UStructProperty* FVariantManagerUtils::GetRelativeScale3DProperty()
 {
 	if (!RelativeScale3DProperty)
 	{
-		RelativeScale3DProperty = FindField<UStructProperty>(USceneComponent::StaticClass(), *USceneComponent::GetRelativeScale3DPropertyName().ToString());
+		RelativeScale3DProperty = FindField<UStructProperty>( USceneComponent::StaticClass(), GET_MEMBER_NAME_CHECKED( USceneComponent, USceneComponent::GetRelativeScale3DPropertyName() ) );
 	}
 
 	return RelativeScale3DProperty;
@@ -96,7 +96,7 @@ UBoolProperty* FVariantManagerUtils::GetVisibilityProperty()
 {
 	if (!VisiblityProperty)
 	{
-		VisiblityProperty = FindField<UBoolProperty>(USceneComponent::StaticClass(), *USceneComponent::GetVisiblePropertyName().ToString());
+		VisiblityProperty = FindField<UBoolProperty>( USceneComponent::StaticClass(), GET_MEMBER_NAME_CHECKED( USceneComponent, USceneComponent::GetVisiblePropertyName() ) );
 	}
 
 	return VisiblityProperty;
