VariantManager: Fix typo on drag and drop tooltip

diff --git a/Engine/Plugins/Enterprise/VariantManager/Source/VariantManager/Private/DisplayNodes/VariantManagerVariantNode.cpp b/Engine/Plugins/Enterprise/VariantManager/Source/VariantManager/Private/DisplayNodes/VariantManagerVariantNode.cpp
--- a/Engine/Plugins/Enterprise/VariantManager/Source/VariantManager/Private/DisplayNodes/VariantManagerVariantNode.cpp
+++ b/Engine/Plugins/Enterprise/VariantManager/Source/VariantManager/Private/DisplayNodes/VariantManagerVariantNode.cpp
@@ -203,7 +203,7 @@ TOptional<EItemDropZone> FVariantManagerVariantNode::CanDrop(const FDragDropEven
 
 		if (NumActorsWeCanAdd > 0)
 		{
-			FText NewHoverText = FText::Format( LOCTEXT("CanDrop_BindActors", "Bind {0} {0}|plural(one=actor,other=actors) to variant '{2}'"),
+			FText NewHoverText = FText::Format( LOCTEXT("CanDrop_BindActors", "Bind {0} {0}|plural(one=actor,other=actors) to variant '{1}'"),
 				NumActorsWeCanAdd,
 				GetVariant().GetDisplayText());
 
