VariantManager: Fix actor node selection clearing when selecting actor components in the details panel

diff --git a/Engine/Plugins/Enterprise/VariantManager/Source/VariantManager/Private/SVariantManager.cpp b/Engine/Plugins/Enterprise/VariantManager/Source/VariantManager/Private/SVariantManager.cpp
--- a/Engine/Plugins/Enterprise/VariantManager/Source/VariantManager/Private/SVariantManager.cpp
+++ b/Engine/Plugins/Enterprise/VariantManager/Source/VariantManager/Private/SVariantManager.cpp
@@ -3381,6 +3381,23 @@ void SVariantManager::OnEditorSelectionChanged(UObject* NewSelection)
 	Selection->GetSelectedObjects(AActor::StaticClass(), SelectedActors);
 	TSet<UObject*> SelectedActorsSet{SelectedActors};
 
+	TArray<UActorComponent*> SelectedComponents;
+	Selection->GetSelectedObjects<UActorComponent>(SelectedComponents);
+	for (UActorComponent* Component : SelectedComponents)
+	{
+		SelectedActorsSet.Add(Component->GetOwner());
+	}
+
+	SelectedActorsSet.Remove(nullptr);
+
+	// No point in clearing our selection the editor doesn't have anything selected either
+	// Sometimes empty selection events seem to fire when selecting the actual actor "row" on its
+	// component tree display
+	if (SelectedActorsSet.Num() == 0)
+	{
+		return;
+	}
+
 	VariantSelection.SuspendBroadcast();
 	for (TSharedRef<FVariantManagerDisplayNode>& DisplayedActor : DisplayedActors)
 	{
