VariantManager: Fix auto-expose of relative location/rotation/scale properties; Fix formatting issues

diff --git a/Engine/Plugins/Enterprise/VariantManager/Source/VariantManager/Private/SVariantManager.cpp b/Engine/Plugins/Enterprise/VariantManager/Source/VariantManager/Private/SVariantManager.cpp
--- a/Engine/Plugins/Enterprise/VariantManager/Source/VariantManager/Private/SVariantManager.cpp
+++ b/Engine/Plugins/Enterprise/VariantManager/Source/VariantManager/Private/SVariantManager.cpp
@@ -2945,7 +2945,7 @@ namespace SVariantManagerUtils
 	{
 		FString PathActorToLeafComponent;
 		FString PathTailProperty;
-	AActor* TargetActor = nullptr;
+		AActor* TargetActor = nullptr;
 		USceneComponent* LeafComponent = nullptr;
 	};
 
@@ -3008,7 +3008,7 @@ namespace SVariantManagerUtils
 				else if (UActorComponent* StartingActorComp = Cast<UActorComponent>(TargetObject))
 				{
 					if (StartingActorComp->GetOwner() == TargetActor)
-	{
+					{
 						// Actor components can't be attached to scene components, so we know we're
 						// attached directly to the actor and there are no other components in the chain
 						ComponentPath = StartingActorComp->GetName() + PATH_DELIMITER;
@@ -3016,29 +3016,31 @@ namespace SVariantManagerUtils
 					}
 				}
 			}
-	}
+		}
 
-	// We need to check if its a blueprint actor or not, as we handle blueprint root component
-	// names a little bit differently
-	bool bIsBlueprintGeneratedClass = ((UObject*)TargetActor->GetClass())->IsA(UBlueprintGeneratedClass::StaticClass());
+		// We need to check if its a blueprint actor or not, as we handle blueprint root component
+		// names a little bit differently
+		bool bIsBlueprintGeneratedClass = ((UObject*)TargetActor->GetClass())->IsA(UBlueprintGeneratedClass::StaticClass());
 
 		// Build component path up to TargetActor, if our Object is actually a SceneComponent
 		USceneComponent* ComponentPointer = LeafComponent;
 		while (ComponentPointer)
-	{
+		{
 			USceneComponent* AttachParent = ComponentPointer->GetAttachParent();
 
-		FString ComponentName;
+			FString ComponentName;
 
 			// We're some form of root component, naming is different
-		if (AttachParent == nullptr || AttachParent->GetOwner() != TargetActor)
-		{
+			if (AttachParent == nullptr || AttachParent->GetOwner() != TargetActor)
+			{
 				// Users can rename of the root component for a blueprint generated class, so lets use that.
 				// On other cases, users can't rename root components, and their actual names are always
 				// something like StaticMeshComponent0 or LightComponent0 (even if its class is a UPointLightComponent).
 				// Getting the class display name matches how the Variant Manager behaves
-				ComponentName = bIsBlueprintGeneratedClass ? ComponentPointer->GetName() :
-															 ComponentPointer->GetClass()->GetDisplayNameText().ToString();
+				ComponentName = bIsBlueprintGeneratedClass ?
+					ComponentPointer->GetName() :
+					ComponentPointer->GetClass()->GetDisplayNameText().ToString();
+
 				ComponentPointer = nullptr;
 			}
 			else
@@ -3053,7 +3055,7 @@ namespace SVariantManagerUtils
 		FString PropertyPath;
 
 		FString ChildPropertyName = ChildProperty ? ChildProperty->GetDisplayNameText().ToString() : FString();
-		static const TSet<FString> ProxyPropertyPaths{ TEXT("Relative Location"), TEXT("Relative Rotation"), TEXT("Relative Scale 3D") };
+		static const TSet<FString> ProxyPropertyPaths{TEXT("Relative Location"), TEXT("Relative Rotation"), TEXT("Relative Scale 3D")};
 
 		// We capture as just 'Material' in the Variant Manager UI, instead of 'Override Materials'
 		// Override Materials doesn't work like a regular UArrayProperty, we need to use GetNumMaterials
@@ -3069,23 +3071,24 @@ namespace SVariantManagerUtils
 		// to be on the root component itself, so as to remove duplicates
 		else if (ComponentPath.IsEmpty() && ProxyPropertyPaths.Contains(ChildPropertyName))
 		{
-			FString RootComponentName = bIsBlueprintGeneratedClass ? TargetActor->GetRootComponent()->GetName() :
-										TargetActor->GetRootComponent()->GetClass()->GetDisplayNameText().ToString();
+			FString RootComponentName = bIsBlueprintGeneratedClass ?
+				TargetActor->GetRootComponent()->GetName() :
+				TargetActor->GetRootComponent()->GetClass()->GetDisplayNameText().ToString();
 
-			ComponentPath = RootComponentName;
+			ComponentPath = RootComponentName + PATH_DELIMITER;
 			PropertyPath += ChildPropertyName;
-			}
-			else
-			{
+		}
+		else
+		{
 			PropertyPath += ChildPropertyName;
-			}
+		}
 
 		Result.TargetActor = TargetActor;
 		Result.LeafComponent = LeafComponent;
 		Result.PathActorToLeafComponent = ComponentPath;
 		Result.PathTailProperty = PropertyPath;
 		return Result;
-		}
+	}
 }
 
 void SVariantManager::OnObjectPropertyChanged(UObject* Object, struct FPropertyChangedEvent& Event)
@@ -3169,7 +3172,7 @@ void SVariantManager::OnPreObjectPropertyChanged(UObject* Object, const class FE
 	// but in the first firing the root component is Object, and in the second firing, the root component is
 
 	if (!bAutoCaptureProperties || !Object)
-		{
+	{
 		return;
 	}
 
@@ -3231,7 +3234,7 @@ void SVariantManager::OnPreObjectPropertyChanged(UObject* Object, const class FE
 	for (PropNode* Node = Head; Node; Node = Node->GetNextNode())
 	{
 		FProperty* Prop = Node->GetValue();
-		if(!Prop)
+		if (!Prop)
 		{
 			return;
 		}
@@ -3247,20 +3250,20 @@ void SVariantManager::OnPreObjectPropertyChanged(UObject* Object, const class FE
 				Category = Category.Replace(TEXT("|"), PATH_DELIMITER);
 
 				int32 LastDelimiterIndex = Category.Find(PATH_DELIMITER, ESearchCase::CaseSensitive, ESearchDir::FromEnd);
-				FString LastCategorySegment = (LastDelimiterIndex == INDEX_NONE)? Category : Category.RightChop(LastDelimiterIndex+1);
+				FString LastCategorySegment = (LastDelimiterIndex == INDEX_NONE) ? Category : Category.RightChop(LastDelimiterIndex + 1);
 
 				if (!IntermediatePath.EndsWith(PATH_DELIMITER + Category + PATH_DELIMITER) &&
 					LastCategorySegment != PropertyName)
-	{
+				{
 					IntermediatePath += Category + PATH_DELIMITER;
 				}
 			}
-	}
+		}
 
 		if (Node == Tail)
-	{
+		{
 			break;
-	}
+		}
 
 		IntermediatePath += PropertyName + PATH_DELIMITER;
 
@@ -3279,7 +3282,7 @@ void SVariantManager::OnPreObjectPropertyChanged(UObject* Object, const class FE
 			}
 		}
 		else
-	{
+		{
 			bShowCategories = false;
 		}
 	}
