VariantManager: Fix crash when pressing plus button to add actors with empty variant sets selected.

Also prevents button from lighting up at all if the variant set has no variants.

diff --git a/Engine/Plugins/Enterprise/VariantManager/Source/VariantManager/Private/SVariantManager.cpp b/Engine/Plugins/Enterprise/VariantManager/Source/VariantManager/Private/SVariantManager.cpp
--- a/Engine/Plugins/Enterprise/VariantManager/Source/VariantManager/Private/SVariantManager.cpp
+++ b/Engine/Plugins/Enterprise/VariantManager/Source/VariantManager/Private/SVariantManager.cpp
@@ -338,7 +338,13 @@ void SVariantManager::Construct(const FArguments& InArgs, TSharedRef<FVariantMan
 								TArray<UVariantSet*> VariantSets;
 								VariantManager->GetSelection().GetSelectedVariantsAndVariantSets(Variants, VariantSets);
 
-								return Variants.Num() + VariantSets.Num() > 0;
+								int32 NumVariants = Variants.Num();
+								for (UVariantSet* VariantSet : VariantSets)
+								{
+									NumVariants += VariantSet->GetNumVariants();
+								}
+
+								return NumVariants > 0;
 							}
 							return false;
 						})
@@ -2625,10 +2631,14 @@ FReply SVariantManager::OnSummonAddActorMenu()
 			}
 		}
 	}
-	TSet<const AActor*>& CommonActorSet = BoundActorSets[0];
-	for (int32 Index = 1; Index < BoundActorSets.Num(); ++Index)
+	TSet<const AActor*> CommonActorSet;
+	if (BoundActorSets.Num() > 0)
 	{
-		CommonActorSet = CommonActorSet.Intersect(BoundActorSets[Index]);
+		CommonActorSet = BoundActorSets[0];
+		for (int32 Index = 1; Index < BoundActorSets.Num(); ++Index)
+		{
+			CommonActorSet = CommonActorSet.Intersect(BoundActorSets[Index]);
+		}
 	}
 
 	auto IsActorValidForAssignment = [CommonActorSet](const AActor* InActor)
