Merging using //UE5/Main_to_//UE5/Release-Engine-Staging [REDACTED_HANDLE]
Files that needed resolving

diff --git a/Engine/Plugins/Enterprise/VariantManager/Source/VariantManager/Private/SVariantManager.cpp b/Engine/Plugins/Enterprise/VariantManager/Source/VariantManager/Private/SVariantManager.cpp
--- a/Engine/Plugins/Enterprise/VariantManager/Source/VariantManager/Private/SVariantManager.cpp
+++ b/Engine/Plugins/Enterprise/VariantManager/Source/VariantManager/Private/SVariantManager.cpp
@@ -18,6 +18,7 @@
 #include "LevelVariantSets.h"
 #include "PropertyValue.h"
 #include "PropertyValueOption.h"
+#include "SDependencyRow.h"
 #include "SVariantManagerActorListView.h"
 #include "SVariantManagerNodeTreeView.h"
 #include "SVariantManagerTableRow.h"
@@ -28,6 +29,7 @@
 #include "VariantManagerEditorCommands.h"
 #include "VariantManagerLog.h"
 #include "VariantManagerSelection.h"
+#include "VariantManagerStyle.h"
 #include "VariantManagerUtils.h"
 #include "VariantSet.h"
 
@@ -63,10 +65,11 @@
 #include "ActorTreeItem.h"
 #include "ScopedTransaction.h"
 #include "Widgets/Input/SButton.h"
-#include "Widgets/Input/SButton.h"
 #include "Widgets/Input/SCheckBox.h"
 #include "Widgets/Input/SSearchBox.h"
 #include "Widgets/Layout/SScrollBorder.h"
+#include "Widgets/Layout/SScrollBox.h"
+#include "Widgets/Layout/SSeparator.h"
 #include "Widgets/Layout/SSplitter.h"
 #include "Widgets/Notifications/SNotificationList.h"
 #include "Widgets/Text/STextBlock.h"
@@ -76,12 +79,15 @@
 
 #define LOCTEXT_NAMESPACE "SVariantManager"
 
+#define VM_COMMON_PADDING 3.0f
+#define VM_COMMON_HEADER_MAX_HEIGHT 26.0f
+
 FSplitterValues::FSplitterValues(FString& InSerialized)
 {
 	TArray<FString> SplitString;
 	InSerialized.ParseIntoArray(SplitString, TEXT(";"));
 
-	if (SplitString.Num() != 4)
+	if (SplitString.Num() != 5)
 	{
 		return;
 	}
@@ -89,7 +95,8 @@
 	VariantColumn = FCString::Atof(*SplitString[0]);
 	ActorColumn = FCString::Atof(*SplitString[1]);
 	PropertyNameColumn = FCString::Atof(*SplitString[2]);
-	PropertyValueColumn = FCString::Atof(*SplitString[3]);
+	DependenciesVariantSetsColumn = FCString::Atof(*SplitString[3]);
+	DependenciesVariantColumn = FCString::Atof(*SplitString[4]);
 }
 
 FString FSplitterValues::ToString()
@@ -97,7 +104,8 @@
 	return FString::SanitizeFloat(VariantColumn) + TEXT(";") +
 		   FString::SanitizeFloat(ActorColumn) + TEXT(";") +
 		   FString::SanitizeFloat(PropertyNameColumn) + TEXT(";") +
-		   FString::SanitizeFloat(PropertyValueColumn);
+		   FString::SanitizeFloat(DependenciesVariantSetsColumn) + TEXT(";") +
+		   FString::SanitizeFloat(DependenciesVariantColumn);
 }
 
 TSharedRef<SWidget> SVariantManager::MakeAddButton()
@@ -152,6 +160,11 @@
 	return nullptr;
 }
 
+void SVariantManager::OnOutlinerNodeSelectionChanged()
+{
+	RefreshActorList();
+}
+
 void SVariantManager::Construct(const FArguments& InArgs, TSharedRef<FVariantManager> InVariantManager)
 {
 	VariantManagerPtr = InVariantManager;
@@ -160,26 +173,46 @@
 
 	CreateCommandBindings();
 
-	const float CommonPadding = 3.f;
-	const float CommonHeaderMaxHeight = 26.0f;
-
 	SAssignNew(NodeTreeView, SVariantManagerNodeTreeView, InVariantManager->GetNodeTree());
 
 	SAssignNew(ActorListView, SVariantManagerActorListView, InVariantManager)
 		.ListItemsSource(&DisplayedActors);
 
-	FSplitterValues SplitterValues;
 	FString SplitterValuesString;
-	if (GConfig->GetString(TEXT("VariantManager"), TEXT("MainSplitterValues"), SplitterValuesString, GEditorPerProjectIni))
+	if (GConfig->GetString(TEXT("VariantManager"), TEXT("SplitterValues"), SplitterValuesString, GEditorPerProjectIni))
 	{
 		SplitterValues = FSplitterValues(SplitterValuesString);
 	}
 
-	RightPropertyColumnWidth = SplitterValues.PropertyValueColumn / (SplitterValues.PropertyValueColumn + SplitterValues.PropertyNameColumn);
+	PropertiesNameColumnWidth = SplitterValues.PropertyNameColumn;
+	PropertiesValueColumnWidth = 1.0f - SplitterValues.ActorColumn - SplitterValues.PropertyNameColumn;
+	DependenciesVariantColumnWidth = SplitterValues.DependenciesVariantColumn;
+	DependenciesControlColumnWidth = 1.0f - SplitterValues.DependenciesVariantSetsColumn - SplitterValues.DependenciesVariantColumn;
 
-	InVariantManager->GetSelection().GetOnOutlinerNodeSelectionChanged().AddSP(this, &SVariantManager::RefreshActorList);
+	PropertiesColumnSizeData.LeftColumnWidth = TAttribute<float>( this, &SVariantManager::OnGetPropertiesActorColumnWidth );
+	PropertiesColumnSizeData.MiddleColumnWidth = TAttribute<float>( this, &SVariantManager::OnGetPropertiesNameColumnWidth );
+	PropertiesColumnSizeData.RightColumnWidth = TAttribute<float>( this, &SVariantManager::OnGetPropertiesValueColumnWidth );
+	PropertiesColumnSizeData.OnFirstSplitterChanged = SSplitter::FOnSlotResized::CreateSP( this, &SVariantManager::OnSetPropertiesNameColumnWidth );
+	PropertiesColumnSizeData.OnSecondSplitterChanged = SSplitter::FOnSlotResized::CreateLambda( [ this ]( float InNewWidth )
+	{
+		// We want the actor column to remain fixed while we move the splitter between property names and values
+		float PropertyComboWidth = PropertiesNameColumnWidth + PropertiesValueColumnWidth;
+		float OldActorColumnWidth = 1.0f - PropertyComboWidth;
+		PropertiesValueColumnWidth = InNewWidth;
+		PropertiesNameColumnWidth = 1.0f - PropertiesValueColumnWidth - OldActorColumnWidth;
+	} );
+
+	DependenciesColumnSizeData.LeftColumnWidth = TAttribute<float>( this, &SVariantManager::OnGetDependenciesVariantSetColumnWidth );
+	DependenciesColumnSizeData.MiddleColumnWidth = TAttribute<float>( this, &SVariantManager::OnGetDependenciesVariantColumnWidth );
+	DependenciesColumnSizeData.RightColumnWidth = TAttribute<float>( this, &SVariantManager::OnGetDependenciesControlColumnWidth );
+	DependenciesColumnSizeData.OnFirstSplitterChanged = SSplitter::FOnSlotResized::CreateSP( this, &SVariantManager::OnSetDependenciesVariantColumnWidth );
+	DependenciesColumnSizeData.OnSecondSplitterChanged = SSplitter::FOnSlotResized::CreateSP( this, &SVariantManager::OnSetDependenciesControlColumnWidth );
+
+	InVariantManager->GetSelection().GetOnOutlinerNodeSelectionChanged().AddSP(this, &SVariantManager::OnOutlinerNodeSelectionChanged);
 	InVariantManager->GetSelection().GetOnActorNodeSelectionChanged().AddSP(this, &SVariantManager::OnActorNodeSelectionChanged);
 
+	UVariant::OnDependenciesUpdated.AddSP(this, &SVariantManager::OnVariantDependenciesUpdated);
+
 	// Subscribe to when objects are modified so that we can auto-resolve when components/array properties are added/removed/renamed
 	OnObjectTransactedHandle = FCoreUObjectDelegates::OnObjectTransacted.AddRaw(this, &SVariantManager::OnObjectTransacted);
 	OnObjectPropertyChangedHandle = FCoreUObjectDelegates::OnObjectPropertyChanged.AddRaw(this, &SVariantManager::OnObjectPropertyChanged);
@@ -194,6 +227,16 @@
 	// new names
 	OnBlueprintCompiledHandle = GEditor->OnBlueprintCompiled().AddSP(this, &SVariantManager::OnBlueprintCompiled);
 
+	OnVariantThumbnailUpdatedHandle = UVariant::OnThumbnailUpdated.AddLambda([this](UVariant* Variant)
+	{
+		OnThumbnailChanged(Variant);
+	});
+
+	OnVariantSetThumbnailUpdatedHandle = UVariantSet::OnThumbnailUpdated.AddLambda([this](UVariantSet* VariantSet)
+	{
+		OnThumbnailChanged(VariantSet);
+	});
+
 	if (FModuleManager::Get().IsModuleLoaded(TEXT("LevelEditor")))
 	{
 		FLevelEditorModule& LevelEditorModule = FModuleManager::LoadModuleChecked<FLevelEditorModule>(TEXT("LevelEditor"));
@@ -202,269 +245,158 @@
 
 	RecordButtonBrush = MakeShared<FSlateImageBrush>(FPaths::EngineContentDir() / TEXT("Editor/Slate/Icons/CA_Record.png"), FVector2D(24.0f, 24.0f));
 
+	RightTreeRootItems.Empty();
+	RightTreeRootItems.Reserve(2);
+	RightTreeRootItems.Add(MakeShared<ERightTreeRowType>(ERightTreeRowType::PropertiesHeader));
+	RightTreeRootItems.Add(MakeShared<ERightTreeRowType>(ERightTreeRowType::DependenciesHeader));
+
+	TSharedPtr<SHeaderRow> CollapsedHeader = SNew( SHeaderRow ).Visibility( EVisibility::Collapsed );
+
+	TSharedRef< STreeView<TSharedRef<ERightTreeRowType>> > RightTree = SNew(STreeView<TSharedRef<ERightTreeRowType>>)
+		.TreeItemsSource( &RightTreeRootItems )
+		.SelectionMode( ESelectionMode::None )
+		.HeaderRow( CollapsedHeader )
+		.AllowOverscroll( EAllowOverscroll::No )
+		.IsEnabled_Lambda([this]()
+		{
+			TSharedPtr<FVariantManager> PinnedVariantManager = VariantManagerPtr.Pin();
+			if (PinnedVariantManager.IsValid())
+			{
+				return PinnedVariantManager->GetSelection().GetSelectedOutlinerNodes().Num() > 0;
+			}
+			return false;
+		})
+		.OnGenerateRow( this, &SVariantManager::GenerateRightTreeRow )
+		.OnGetChildren_Lambda([this]( TSharedRef<ERightTreeRowType> InRowType, TArray<TSharedRef<ERightTreeRowType>>& OutChildren )
+		{
+			OutChildren.Empty();
+			if ( *InRowType == ERightTreeRowType::PropertiesHeader )
+			{
+				OutChildren.Add( MakeShared<ERightTreeRowType>( ERightTreeRowType::PropertiesContent ) );
+			}
+			else if ( *InRowType == ERightTreeRowType::DependenciesHeader )
+			{
+				OutChildren.Add( MakeShared<ERightTreeRowType>( ERightTreeRowType::DependenciesContent ) );
+			}
+		});
+
+	// Expand right tree by default
+	for ( const TSharedRef<ERightTreeRowType>& Item : RightTreeRootItems )
+	{
+		RightTree->SetItemExpansion(Item, true);
+	}
+
+	float BorderThickness = FVariantManagerStyle::Get()->GetFloat( "VariantManager.Spacings.BorderThickness" );
+
 	ChildSlot
 	[
-		SAssignNew(MainSplitter, SSplitter)
-		.Orientation(Orient_Horizontal)
-		// VariantSet/Variant column
-		+SSplitter::Slot()
-		.Value(SplitterValues.VariantColumn)
+		SNew(SVerticalBox)
+
+		// Main header bar (+ Variant Set button and etc)
+		+SVerticalBox::Slot()
+		.MaxHeight(VM_COMMON_HEADER_MAX_HEIGHT)
+		.AutoHeight()
+		.Padding(FMargin(VM_COMMON_PADDING, VM_COMMON_PADDING, 0.0f, VM_COMMON_PADDING))
 		[
-			SNew(SVerticalBox)
+			SNew(SBox)
+			.HeightOverride(VM_COMMON_HEADER_MAX_HEIGHT)
+			.HAlign(HAlign_Fill)
+			.VAlign(VAlign_Fill)
+			[
+				SNew(SHorizontalBox)
+				+SHorizontalBox::Slot()
+				.VAlign(VAlign_Center)
+				.Padding(FMargin(0.f, 0.f, VM_COMMON_PADDING+2.0f, 1.f))
+				.AutoWidth()
+				[
+					MakeAddButton()
+				]
 
-			// +VariantSets button and search
-			+SVerticalBox::Slot()
-			.MaxHeight(CommonHeaderMaxHeight)
-			.AutoHeight()
-			.Padding(FMargin(CommonPadding, CommonPadding, 0.0f, CommonPadding))
-			[
-				SNew(SBox)
-				.HeightOverride(CommonHeaderMaxHeight)
-				.HAlign(HAlign_Fill)
-				.VAlign(VAlign_Fill)
+				+SHorizontalBox::Slot()
+				.VAlign(VAlign_Center)
+				.Padding(FMargin(0.f, 0.f, VM_COMMON_PADDING+2.0f, 0.f))
+				.AutoWidth()
+				.MaxWidth(VM_COMMON_HEADER_MAX_HEIGHT) // square aspect ratio
 				[
-					SNew(SHorizontalBox)
-					+SHorizontalBox::Slot()
-					.VAlign(VAlign_Center)
-					.Padding(FMargin(0.f, 0.f, CommonPadding+2.0f, 1.f))
-					.AutoWidth()
+					SNew(SBox)
+					.HeightOverride(VM_COMMON_HEADER_MAX_HEIGHT - 8.0f) // These so that it matches the height of the search box
+					.WidthOverride(VM_COMMON_HEADER_MAX_HEIGHT - 8.0f)
 					[
-						MakeAddButton()
-					]
-
-					+SHorizontalBox::Slot()
-					.VAlign(VAlign_Center)
-					.Padding(FMargin(0.f, 0.f, CommonPadding+2.0f, 0.f))
-					.AutoWidth()
-					.MaxWidth(CommonHeaderMaxHeight) // square aspect ratio
-					[
-						SNew(SBox)
-						.HeightOverride(CommonHeaderMaxHeight - 8.0f) // These so that it matches the height of the search box
-						.WidthOverride(CommonHeaderMaxHeight - 8.0f)
+						SNew(SCheckBox)
+						.Style(FCoreStyle::Get(), "ToggleButtonCheckbox")
+						.ToolTipText(LOCTEXT("AutoCaptureTooltip", "Enable or disable auto-capturing properties"))
+						.IsChecked_Lambda([&bAutoCaptureProperties = bAutoCaptureProperties]()
+						{
+							return bAutoCaptureProperties? ECheckBoxState::Checked : ECheckBoxState::Unchecked;
+						})
+						.OnCheckStateChanged_Lambda([&bAutoCaptureProperties = bAutoCaptureProperties](const ECheckBoxState NewState)
+						{
+							bAutoCaptureProperties = NewState == ECheckBoxState::Checked;
+						})
 						[
-							SNew(SCheckBox)
-							.Style(FCoreStyle::Get(), "ToggleButtonCheckbox")
-							.ToolTipText(LOCTEXT("AutoCaptureTooltip", "Enable or disable auto-capture properties"))
-							.IsChecked_Lambda([&bAutoCaptureProperties = bAutoCaptureProperties]()
-							{
-								return bAutoCaptureProperties? ECheckBoxState::Checked : ECheckBoxState::Unchecked;
-							})
-							.OnCheckStateChanged_Lambda([&bAutoCaptureProperties = bAutoCaptureProperties](const ECheckBoxState NewState)
-							{
-								bAutoCaptureProperties = NewState == ECheckBoxState::Checked;
-							})
+							SNew(SBox)
+							.Padding(FMargin(0.0f, 2.0, 2.0f, 2.0)) // Extra padding on the right because ToggleButtonCheckboxes always nudges the image to the right
 							[
-								SNew(SBox)
-								.Padding(FMargin(0.0f, 2.0, 2.0f, 2.0)) // Extra padding on the right because ToggleButtonCheckboxes always nudges the image to the right
-								[
-									SNew(SImage)
-									.Image(RecordButtonBrush.Get())
-								]
+								SNew(SImage)
+								.Image(RecordButtonBrush.Get())
 							]
 						]
 					]
+				]
 
-					+SHorizontalBox::Slot()
-					.FillWidth(1.0f)
-					.VAlign(VAlign_Center)
-					[
-						SNew(SSearchBox)
-						.HintText(LOCTEXT("VariantManagerFilterText", "Filter"))
-						.OnTextChanged(this, &SVariantManager::OnOutlinerSearchChanged)
-					]
-				]
-			]
-
-			+SVerticalBox::Slot()
-			.Padding(FMargin(CommonPadding, 0.0f, 0.0f, CommonPadding))
-			.VAlign(VAlign_Fill)
-			.FillHeight(1.0f)
-			[
-				SNew(SScrollBorder, NodeTreeView.ToSharedRef())
+				+SHorizontalBox::Slot()
+				.FillWidth(1.0f)
+				.VAlign(VAlign_Center)
+				.MaxWidth( 200.0f )
 				[
-					SNew(SBox) // Very important to prevent the tree from expanding freely
-					[
-						NodeTreeView.ToSharedRef()
-					]
+					SNew(SSearchBox)
+					.HintText(LOCTEXT("VariantManagerFilterText", "Filter variants"))
+					.OnTextChanged(this, &SVariantManager::OnOutlinerSearchChanged)
 				]
 			]
 		]
 
-		// Actor column
-		+ SSplitter::Slot()
-		.Value(SplitterValues.ActorColumn)
+		+ SVerticalBox::Slot()
+		.VAlign( VAlign_Fill )
+		.FillHeight( 1.0f )
 		[
-			SNew(SVerticalBox)
-			+ SVerticalBox::Slot()
-			.MaxHeight(CommonHeaderMaxHeight)
-			.AutoHeight()
-			.Padding(FMargin(0.0f, CommonPadding, 0.0f, CommonPadding))
+			// Common "background"
+			SNew(SBorder)
+			.BorderImage( FEditorStyle::GetBrush( "WhiteBrush" ) )
+			.BorderBackgroundColor( FVariantManagerStyle::Get()->GetColor("VariantManager.Panels.LightBackgroundColor") )
+			.Padding( FMargin( BorderThickness ) )
 			[
-				SNew(SBox)
-				.HeightOverride(CommonHeaderMaxHeight)
-				.HAlign(HAlign_Center)
-				.VAlign(VAlign_Center)
-				[
-					SNew(SHorizontalBox)
-
-					+SHorizontalBox::Slot()
-					.AutoWidth()
-					.VAlign(VAlign_Center)
-					[
-					SNew(STextBlock)
-					.TextStyle(FEditorStyle::Get(), "NormalText.Important")
-					.Text(LOCTEXT("ActorsText", "Actors"))
-				]
-
-					+SHorizontalBox::Slot()
-					.VAlign(VAlign_Center)
-					.HAlign(HAlign_Right)
-					.MaxWidth(24.0f)
-					.AutoWidth()
-					.Padding(FMargin(4.f, 0.f, 4.f, 0.f))
-					[
-						SNew(SButton)
-						.ButtonStyle(FEditorStyle::Get(), "NoBorder")
-						.ToolTipText(LOCTEXT("AddActorPlusTooltip", "Add a new actor binding to selected variants"))
-						.OnClicked(this, &SVariantManager::OnSummonAddActorMenu)
-						.ContentPadding(FMargin(2.0f, 1.0f))
-						.IsEnabled_Lambda([&VariantManagerPtr = VariantManagerPtr]() -> bool
-						{
-							if (TSharedPtr<FVariantManager> VariantManager = VariantManagerPtr.Pin())
-							{
-								TArray<UVariant*> Variants;
-								TArray<UVariantSet*> VariantSets;
-								VariantManager->GetSelection().GetSelectedVariantsAndVariantSets(Variants, VariantSets);
-
-								if (Variants.Num() > 0)
-								{
-									return true;
-								}
-
-								for (const UVariantSet* Set : VariantSets)
-								{
-									if (Set && Set->GetNumVariants() > 0)
-									{
-										return true;
-									}
-								}
-							}
-							return false;
-						})
-						.Content()
-						[
-							SNew(STextBlock)
-							.TextStyle(FEditorStyle::Get(), "NormalText.Important")
-							.Font(FEditorStyle::Get().GetFontStyle("FontAwesome.10"))
-							.Text(FEditorFontGlyphs::Plus)
-						]
-					]
-				]
-			]
-
-			+ SVerticalBox::Slot()
-			.Padding(FMargin(0.0f, 0.0f, 0.0f, CommonPadding))
-			.VAlign(VAlign_Fill)
-			.FillHeight(1.0f)
-			[
-				ActorListView.ToSharedRef()
-			]
-		]
-
-		// Properties column
-		+ SSplitter::Slot()
-		.Value(SplitterValues.PropertyNameColumn + SplitterValues.PropertyValueColumn)
-		[
-			SNew(SVerticalBox)
-			+ SVerticalBox::Slot()
-			.MaxHeight(CommonHeaderMaxHeight)
-			.AutoHeight()
-			.Padding(0.0f, CommonPadding, CommonPadding, CommonPadding)
-			[
-				// Headers
-				SNew(SSplitter)
+				SAssignNew(MainSplitter, SSplitter)
 				.Orientation(Orient_Horizontal)
-				+ SSplitter::Slot()
-				.Value(ColumnSizeData.NameColumnWidth)
-				.OnSlotResized(ColumnSizeData.OnNameColumnResized)
+				.PhysicalSplitterHandleSize( BorderThickness )
+				.HitDetectionSplitterHandleSize( BorderThickness )
+
+				// VariantSet/Variant column
+				+SSplitter::Slot()
+				.Value(SplitterValues.VariantColumn)
 				[
-					SNew(SBox)
-					.HeightOverride(CommonHeaderMaxHeight)
-					.HAlign(HAlign_Center)
-					.VAlign(VAlign_Center)
+					SNew(SScrollBorder, NodeTreeView.ToSharedRef())
 					[
-						SNew(SHorizontalBox)
-
-						+SHorizontalBox::Slot()
-						.AutoWidth()
-						.VAlign(VAlign_Center)
+						SNew(SBox) // Very important to prevent the tree from expanding freely
 						[
-						SNew(STextBlock)
-						.TextStyle(FEditorStyle::Get(), "NormalText.Important")
-						.Text(LOCTEXT("PropertiesText", "Properties"))
-					]
-
-						+SHorizontalBox::Slot()
-						.VAlign(VAlign_Center)
-						.HAlign(HAlign_Right)
-						.MaxWidth(24.0f)
-						.AutoWidth()
-						.Padding(FMargin(4.f, 0.f, 4.f, 0.f))
-						[
-							SNew(SButton)
-							.ButtonStyle(FEditorStyle::Get(), "NoBorder")
-							.ToolTipText(LOCTEXT("CapturePropertiesPlusTooltip", "Capture properties from the selected actor bindings"))
-							.OnClicked_Lambda([this]
-							{
-								CaptureNewPropertiesFromSelectedActors();
-								return FReply::Handled();
-							})
-							.ContentPadding(FMargin(2.0f, 1.0f))
-							.IsEnabled_Lambda([&VariantManagerPtr = VariantManagerPtr]() -> bool
-							{
-								if (TSharedPtr<FVariantManager> VariantManager = VariantManagerPtr.Pin())
-								{
-									return VariantManager->GetSelection().GetSelectedActorNodes().Num() > 0;
-								}
-								return false;
-							})
-							.Content()
-							[
-								SNew(STextBlock)
-								.TextStyle(FEditorStyle::Get(), "NormalText.Important")
-								.Font(FEditorStyle::Get().GetFontStyle("FontAwesome.10"))
-								.Text(FEditorFontGlyphs::Plus)
-							]
+							NodeTreeView.ToSharedRef()
 						]
 					]
 				]
+
+				// Properties/Dependencies column
 				+ SSplitter::Slot()
-				.Value(ColumnSizeData.RightColumnWidth)
-				.OnSlotResized(ColumnSizeData.OnValueColumnResized)
+				.Value( 1.0f - SplitterValues.VariantColumn )
 				[
-					SNew(SBox)
-					.HeightOverride(CommonHeaderMaxHeight)
-					.HAlign(HAlign_Center)
-					.VAlign(VAlign_Center)
+					SNew( SScrollBorder, RightTree )
 					[
-						SNew(STextBlock)
-						.TextStyle(FEditorStyle::Get(), "NormalText.Important")
-						.Text(LOCTEXT("PropertiesValuesText", "Values"))
+						SNew( SBox ) // Very important to prevent the tree from expanding freely
+						[
+							RightTree
+						]
 					]
 				]
 			]
-			+ SVerticalBox::Slot()
-			.Padding(0.0f, 0.0f, CommonPadding, CommonPadding)
-			.VAlign(VAlign_Fill)
-			.FillHeight(1.0f)
-			[
-				SAssignNew(CapturedPropertyListView, SListView<TSharedPtr<FVariantManagerPropertyNode>>)
-				.SelectionMode(ESelectionMode::Single)
-				.ListItemsSource(&DisplayedPropertyNodes)
-				.OnContextMenuOpening(this, &SVariantManager::OnPropertyListContextMenuOpening)
-				.OnGenerateRow(this, &SVariantManager::MakeCapturedPropertyRow)
-				.Visibility(EVisibility::Visible)
-			]
 		]
 	];
 
@@ -496,6 +428,14 @@
 	USelection::SelectionChangedEvent.Remove(OnEditorSelectionChangedHandle);
 	OnEditorSelectionChangedHandle.Reset();
 
+	UVariant::OnDependenciesUpdated.RemoveAll(this);
+
+	UVariant::OnThumbnailUpdated.Remove(OnVariantThumbnailUpdatedHandle);
+	OnVariantThumbnailUpdatedHandle.Reset();
+
+	UVariantSet::OnThumbnailUpdated.Remove(OnVariantSetThumbnailUpdatedHandle);
+	OnVariantSetThumbnailUpdatedHandle.Reset();
+
 	if (FModuleManager::Get().IsModuleLoaded(TEXT("LevelEditor")))
 	{
 		FLevelEditorModule& LevelEditorModule = FModuleManager::LoadModuleChecked<FLevelEditorModule>(TEXT("LevelEditor"));
@@ -504,18 +444,19 @@
 	}
 
 	// Save splitter layout
+	if ( MainSplitter )
 	{
-		FChildren* Slots = MainSplitter->GetChildren();
-		if (Slots->Num() == 3)
+		FChildren* MainSlots = MainSplitter->GetChildren();
+		if (MainSlots->Num() > 0)
 		{
 			FSplitterValues Values;
 			Values.VariantColumn = MainSplitter->SlotAt(0).SizeValue.Get();
-			Values.ActorColumn = MainSplitter->SlotAt(1).SizeValue.Get();
-			float PropertyCombo = MainSplitter->SlotAt(2).SizeValue.Get();
-			Values.PropertyNameColumn = PropertyCombo * ColumnSizeData.NameColumnWidth.Get();
-			Values.PropertyValueColumn = PropertyCombo * ColumnSizeData.ValueColumnWidth.Get();
+			Values.ActorColumn = OnGetPropertiesActorColumnWidth();;
+			Values.PropertyNameColumn = OnGetPropertiesNameColumnWidth();
+			Values.DependenciesVariantSetsColumn = OnGetDependenciesVariantSetColumnWidth();
+			Values.DependenciesVariantColumn = OnGetDependenciesVariantColumnWidth();
 
-			GConfig->SetString(TEXT("VariantManager"), TEXT("MainSplitterValues"), *Values.ToString(), GEditorPerProjectIni);
+			GConfig->SetString(TEXT("VariantManager"), TEXT("SplitterValues"), *Values.ToString(), GEditorPerProjectIni);
 		}
 	}
 }
@@ -575,19 +516,19 @@
 	);
 
 	VariantTreeCommandBindings->MapAction(
-		FVariantManagerEditorCommands::Get().CreateThumbnailVariantCommand,
+		FVariantManagerEditorCommands::Get().CreateThumbnailCommand,
 		FExecuteAction::CreateSP(this, &SVariantManager::CreateThumbnail),
 		FCanExecuteAction::CreateSP(this, &SVariantManager::CanCreateThumbnail)
 	);
 
 	VariantTreeCommandBindings->MapAction(
-		FVariantManagerEditorCommands::Get().LoadThumbnailVariantCommand,
+		FVariantManagerEditorCommands::Get().LoadThumbnailCommand,
 		FExecuteAction::CreateSP(this, &SVariantManager::LoadThumbnail),
 		FCanExecuteAction::CreateSP(this, &SVariantManager::CanLoadThumbnail)
 	);
 
 	VariantTreeCommandBindings->MapAction(
-		FVariantManagerEditorCommands::Get().ClearThumbnailVariantCommand,
+		FVariantManagerEditorCommands::Get().ClearThumbnailCommand,
 		FExecuteAction::CreateSP(this, &SVariantManager::ClearThumbnail),
 		FCanExecuteAction::CreateSP(this, &SVariantManager::CanClearThumbnail)
 	);
@@ -1249,64 +1190,24 @@
 		return;
 	}
 
-	FViewport* Viewport = GEditor->GetActiveViewport();
-
-	if ( !ensure(GCurrentLevelEditingViewportClient) || !ensure(Viewport) )
-	{
-		return;
-		}
-
-	// Copied from FAssetContextMenu::ExecuteCaptureThumbnail
-		//have to re-render the requested viewport
-		FLevelEditorViewportClient* OldViewportClient = GCurrentLevelEditingViewportClient;
-		//remove selection box around client during render
-		GCurrentLevelEditingViewportClient = nullptr;
-		Viewport->Draw();
-
-	// Copied from FContentBrowserSingleton::CaptureThumbnailFromViewport
-	//capture the thumbnail
-	uint32 SrcWidth = Viewport->GetSizeXY().X;
-	uint32 SrcHeight = Viewport->GetSizeXY().Y;
-	// Read the contents of the viewport into an array.
-	TArray<FColor> OrigBitmap;
-	if (!Viewport->ReadPixels(OrigBitmap) || OrigBitmap.Num() != SrcWidth * SrcHeight)
-	{
-		return;
-	}
-
-	TArray<FColor> ScaledBitmap;
-	FImageUtils::CropAndScaleImage(SrcWidth, SrcHeight, VARIANT_THUMBNAIL_SIZE, VARIANT_THUMBNAIL_SIZE, OrigBitmap, ScaledBitmap);
-
-		//redraw viewport to have the yellow highlight again
-		GCurrentLevelEditingViewportClient = OldViewportClient;
-		Viewport->Draw();
-
 	TArray<UVariant*> Vars;
 	TArray<UVariantSet*> VarSets;
 	VariantManager->GetSelection().GetSelectedVariantsAndVariantSets(Vars, VarSets);
-	if (Vars.Num() < 1)
+	if (Vars.Num() + VarSets.Num() < 1)
 	{
 		return;
 	}
 
-	FScopedTransaction Transaction(LOCTEXT("CaptureThumbnailTransaction", "Capture viewport as a variant thumbnail"));
+	FScopedTransaction Transaction(LOCTEXT("CaptureThumbnailTransaction", "Capture viewport as a thumbnail for a variant or variant set"));
 
-	FCreateTexture2DParameters Params;
-	Params.bDeferCompression = true;
-
 	for (UVariant* Var : Vars)
 	{
-		// Create one thumbnail for each so that the Variants can own them
-		UTexture2D* ResizedTexture = FImageUtils::CreateTexture2D(VARIANT_THUMBNAIL_SIZE, VARIANT_THUMBNAIL_SIZE, ScaledBitmap, Var, FString(), RF_NoFlags, Params);
-		if (!ResizedTexture)
-		{
-			continue;
-		}
-
-		Var->SetThumbnail(ResizedTexture);
+		Var->SetThumbnailFromEditorViewport();
 	}
-
-	RefreshVariantTree();
+	for (UVariantSet* VarSet : VarSets)
+	{
+		VarSet->SetThumbnailFromEditorViewport();
+	}
 }
 
 void SVariantManager::LoadThumbnail()
@@ -1323,6 +1224,14 @@
 		return;
 	}
 
+	TArray<UVariant*> Vars;
+	TArray<UVariantSet*> VarSets;
+	VariantManager->GetSelection().GetSelectedVariantsAndVariantSets(Vars, VarSets);
+	if (Vars.Num() + VarSets.Num() < 1)
+	{
+		return;
+	}
+
 	// Formats that IImageWrapper can handle
 	const FString Filter = TEXT("Image files (*.jpg; *.png; *.bmp; *.ico; *.exr; *.icns)|*.jpg; *.png; *.bmp; *.ico; *.exr; *.icns|All files (*.*)|*.*");
 	TArray<FString> OutFiles;
@@ -1338,35 +1247,19 @@
 	}
 
 	FString SourceImagePath = FPaths::ConvertRelativePathToFull(OutFiles[0]);
-	if (!FPlatformFileManager::Get().GetPlatformFile().FileExists(*SourceImagePath))
-	{
-		return;
-	}
 
-	UTexture2D* OriginalTexture = FImageUtils::ImportFileAsTexture2D(SourceImagePath);
-	if (!OriginalTexture)
-	{
-		return;
-	}
-
-	TArray<UVariant*> Vars;
-	TArray<UVariantSet*> VarSets;
-	VariantManager->GetSelection().GetSelectedVariantsAndVariantSets(Vars, VarSets);
-	if (Vars.Num() < 1)
-	{
-		return;
-	}
-
 	FScopedTransaction Transaction(FText::Format(
-		LOCTEXT("LoadThumbnailTransaction", "Load file '{0}' as a variant thumbnail"), FText::FromString(SourceImagePath)
+		LOCTEXT("LoadThumbnailTransaction", "Load file '{0}' as a thumbnail for a variant or variant set"), FText::FromString(SourceImagePath)
 	));
 
 	for (UVariant* Var : Vars)
 	{
-		Var->SetThumbnail(OriginalTexture);
+		Var->SetThumbnailFromFile(SourceImagePath);
 	}
-
-	RefreshVariantTree();
+	for (UVariantSet* VarSet : VarSets)
+	{
+		VarSet->SetThumbnailFromFile(SourceImagePath);
+	}
 }
 
 void SVariantManager::ClearThumbnail()
@@ -1380,7 +1273,7 @@
 	TArray<UVariant*> SelectedVariants;
 	TArray<UVariantSet*> SelectedVariantSets;
 	VarMan->GetSelection().GetSelectedVariantsAndVariantSets(SelectedVariants, SelectedVariantSets);
-	if (SelectedVariants.Num() < 1)
+	if (SelectedVariants.Num() + SelectedVariantSets.Num() < 1)
 	{
 		return;
 	}
@@ -1389,10 +1282,12 @@
 
 	for (UVariant* Variant : SelectedVariants)
 	{
-		Variant->SetThumbnail(nullptr);
+		Variant->SetThumbnailFromTexture(nullptr);
 	}
-
-	RefreshVariantTree();
+	for (UVariantSet* VariantSet : SelectedVariantSets)
+	{
+		VariantSet->SetThumbnailFromTexture(nullptr);
+	}
 }
 
 bool SVariantManager::CanSwitchOnVariant()
@@ -2071,6 +1966,30 @@
 	RefreshPropertyList();
 }
 
+void SVariantManager::OnVariantDependenciesUpdated(UVariant* ParentVariant)
+{
+	if (!ParentVariant)
+	{
+		return;
+	}
+
+	TSharedPtr<FVariantManager> PinnedVariantManager = VariantManagerPtr.Pin();
+	if (!PinnedVariantManager.IsValid())
+	{
+		return;
+	}
+
+	TArray<UVariant*> SelectedVariants;
+	TArray<UVariantSet*> SelectedVariantSets;
+	PinnedVariantManager->GetSelection().GetSelectedVariantsAndVariantSets(SelectedVariants, SelectedVariantSets);
+	if (!SelectedVariants.Contains(ParentVariant))
+	{
+		return;
+	}
+
+	RefreshDependencyLists();
+}
+
 void SVariantManager::RefreshVariantTree()
 {
 	TSharedPtr<FVariantManager> PinnedVariantManager = VariantManagerPtr.Pin();
@@ -2250,6 +2169,7 @@
 	// the captured properties, because we might select a different variant now, so the captured
 	// properties could be different
 	RefreshPropertyList();
+	RefreshDependencyLists();
 }
 
 namespace SVariantManagerUtils
@@ -2315,6 +2235,12 @@
 
 void SVariantManager::RefreshPropertyList()
 {
+	// Don't update if the properties tab is collapsed
+	if ( CapturedPropertyListView == nullptr )
+	{
+		return;
+	}
+
 	FVariantManagerSelection& Selection = VariantManagerPtr.Pin()->GetSelection();
 
 	TArray<UPropertyValue*> Properties;
@@ -2426,6 +2352,53 @@
 	CapturedPropertyListView->RequestListRefresh();
 }
 
+void SVariantManager::RefreshDependencyLists()
+{
+	if (!DependenciesList || !DependentsList)
+	{
+		return;
+	}
+
+	TSharedPtr<FVariantManager> PinnedVariantManager = VariantManagerPtr.Pin();
+	if (!PinnedVariantManager.IsValid())
+	{
+		return;
+	}
+
+	DisplayedDependencies.Reset();
+	DisplayedDependents.Reset();
+
+	TArray<UVariant*> SelectedVariants;
+	TArray<UVariantSet*> SelectedVariantSets;
+	PinnedVariantManager->GetSelection().GetSelectedVariantsAndVariantSets(SelectedVariants, SelectedVariantSets);
+	if (SelectedVariants.Num() == 1)
+	{
+		UVariant* SelectedVariant = SelectedVariants[0];
+
+		for (int32 Index = 0; Index < SelectedVariant->GetNumDependencies(); ++Index)
+		{
+			FVariantDependencyModelPtr Model = MakeShared<FVariantDependencyModel>();
+			Model->Dependency = &SelectedVariant->GetDependency(Index);
+			Model->ParentVariant = SelectedVariant;
+
+			DisplayedDependencies.Add(Model);
+		}
+
+		const bool bOnlyEnabledDependencies = false;
+		TArray<UVariant*> Dependents = SelectedVariant->GetDependents(PinnedVariantManager->GetCurrentLevelVariantSets(), bOnlyEnabledDependencies);
+		for (UVariant* Dependent : Dependents)
+		{
+			FVariantDependencyModelPtr Model = MakeShared<FVariantDependencyModel>();
+			Model->ParentVariant = Dependent;
+
+			DisplayedDependents.Add(Model);
+		}
+	}
+
+	DependenciesList->RequestListRefresh();
+	DependentsList->RequestListRefresh();
+}
+
 void SVariantManager::UpdatePropertyDefaults()
 {
 	TSharedPtr<FVariantManager> PinnedVariantManager = VariantManagerPtr.Pin();
@@ -2695,6 +2668,33 @@
 	return FReply::Handled();
 }
 
+FReply SVariantManager::OnAddDependencyClicked()
+{
+	TSharedPtr<FVariantManager> PinnedVariantManager = VariantManagerPtr.Pin();
+	if (!PinnedVariantManager.IsValid())
+	{
+		return FReply::Handled();
+	}
+
+	TArray<UVariant*> Variants;
+	TArray<UVariantSet*> VariantSets;
+	PinnedVariantManager->GetSelection().GetSelectedVariantsAndVariantSets(Variants, VariantSets);
+	if (Variants.Num() == 1 && Variants[0] != nullptr)
+	{
+		FScopedTransaction Transaction(FText::Format(
+			LOCTEXT("AddDependencyTransaction", "Add a dependency to variant '{0}'"),
+			Variants[0]->GetDisplayText()
+		));
+
+		UVariant* SelectedVariant = Variants[0];
+
+		FVariantDependency Dependency;
+		SelectedVariant->AddDependency(Dependency);
+	}
+
+	return FReply::Handled();
+}
+
 // Tries capturing and recording new data for the property at PropertyPath for TargetActor, into whatever Variants we have selected.
 // Will return true if it created or updated a UPropertyValue
 bool AutoCaptureProperty(FVariantManager* VarMan, AActor* TargetActor, const FString& PropertyPath, const FProperty* Property)
@@ -3367,6 +3367,488 @@
 	RefreshPropertyList();
 }
 
+TSharedRef<ITableRow> SVariantManager::GenerateRightTreeRow( TSharedRef<ERightTreeRowType> RowType, const TSharedRef<STableViewBase>& OwnerTable )
+{
+	TSharedRef<STableRow<TSharedRef<ERightTreeRowType>>> TableRow = SNew( STableRow<TSharedRef<ERightTreeRowType>>, OwnerTable );
+
+	TSharedPtr<SWidget> RowContent = SNullWidget::NullWidget;
+	switch ( *RowType )
+	{
+	case ERightTreeRowType::PropertiesHeader:
+	case ERightTreeRowType::DependenciesHeader:
+		RowContent = GenerateRightTreeHeaderRowContent(*RowType, TableRow);
+		break;
+	case ERightTreeRowType::PropertiesContent:
+		RowContent = GenerateRightTreePropertiesRowContent();
+		break;
+	case ERightTreeRowType::DependenciesContent:
+		RowContent = GenerateRightTreeDependenciesRowContent();
+		break;
+	default:
+		break;
+	}
+
+	TableRow->SetRowContent(RowContent.ToSharedRef());
+	return TableRow;
+}
+
+TSharedRef<SWidget> SVariantManager::GenerateRightTreeHeaderRowContent( ERightTreeRowType RowType, TSharedRef<STableRow<TSharedRef<ERightTreeRowType>>> InTableRow )
+{
+	FText RowText;
+	switch ( RowType )
+	{
+	case ERightTreeRowType::DependenciesHeader:
+		RowText = FText::FromString(TEXT("Dependencies"));
+		break;
+	case ERightTreeRowType::PropertiesHeader:
+		RowText = FText::FromString(TEXT("Properties"));
+		break;
+	default:
+		return SNullWidget::NullWidget;
+	}
+
+	TSharedRef<SHorizontalBox> HorizontalBox = SNew( SHorizontalBox )
+	.IsEnabled_Lambda([this, RowType]()
+	{
+		if (RowType == ERightTreeRowType::PropertiesHeader)
+		{
+			return true;
+		}
+
+		TSharedPtr<FVariantManager> PinnedVariantManager = VariantManagerPtr.Pin();
+		if (PinnedVariantManager.IsValid())
+		{
+			TArray<UVariant*> SelectedVariants;
+			TArray<UVariantSet*> SelectedVariantSets;
+			PinnedVariantManager->GetSelection().GetSelectedVariantsAndVariantSets(SelectedVariants, SelectedVariantSets);
+
+			return SelectedVariants.Num() == 1;
+		}
+		return false;
+	});
+
+	HorizontalBox->AddSlot()
+	.Padding( FMargin( 4.f, 0.f, 4.f, 0.f ) )
+	.VAlign( VAlign_Center )
+	.MaxWidth( 24.0f )
+	[
+		SNew( SExpanderArrow, InTableRow ).IndentAmount( FVariantManagerStyle::Get()->GetFloat("VariantManager.Spacings.IndentAmount") )
+	];
+
+	HorizontalBox->AddSlot()
+		.VAlign( VAlign_Center )
+		.HAlign( HAlign_Left )
+		.AutoWidth()
+		.Padding( FMargin( 0.f, 0.f, 4.f, 0.f ) )
+		[
+			SNew( STextBlock )
+			.TextStyle( FEditorStyle::Get(), "NormalText.Important" )
+			.Text( RowText )
+		];
+
+	if ( RowType == ERightTreeRowType::DependenciesHeader )
+	{
+		HorizontalBox->AddSlot()
+		.VAlign(VAlign_Center)
+		.HAlign(HAlign_Left)
+		.MaxWidth(24.0f)
+		.Padding(FMargin(4.f, 0.f, 4.f, 0.f))
+		[
+			SNew(SButton)
+			.ButtonStyle(FEditorStyle::Get(), "NoBorder")
+			.ToolTipText(LOCTEXT("AddDependencyTooltip", "Add a new dependency to this variant"))
+			.OnClicked(this, &SVariantManager::OnAddDependencyClicked)
+			.ContentPadding(FMargin(2.0f, 1.0f))
+			.Content()
+			[
+				SNew(STextBlock)
+				.TextStyle(FEditorStyle::Get(), "NormalText.Important")
+				.Font(FEditorStyle::Get().GetFontStyle("FontAwesome.10"))
+				.Text(FEditorFontGlyphs::Plus)
+			]
+		];
+	}
+
+	return
+	// Show a white spacer above the dependencies header row
+	SNew( SBorder )
+	.BorderImage( FEditorStyle::GetBrush( "WhiteBrush" ) )
+	.BorderBackgroundColor( FVariantManagerStyle::Get()->GetColor( "VariantManager.Panels.LightBackgroundColor" ) )
+	.Padding( FMargin( 0.0f, RowType == ERightTreeRowType::DependenciesHeader ? FVariantManagerStyle::Get()->GetFloat( "VariantManager.Spacings.BorderThickness" ) : 0.0f, 0.0f, 0.0f ) )
+	[
+		SNew( SBox )
+		.HeightOverride( 26 )
+		[
+			SNew( SBorder )
+			.VAlign( VAlign_Center )
+			.HAlign( HAlign_Fill )
+			.BorderImage( FEditorStyle::GetBrush( "Sequencer.AnimationOutliner.TopLevelBorder_Expanded" ) )
+			.BorderBackgroundColor( FVariantManagerStyle::Get()->GetColor( "VariantManager.Panels.HeaderBackgroundColor" ) )
+			.Padding( FMargin( 2.0f, 0.0f, 2.0f, 0.0f ) )
+			[
+				HorizontalBox
+			]
+		]
+	];
+}
+
+TSharedRef<SWidget> SVariantManager::GenerateRightTreePropertiesRowContent()
+{
+	float BorderThickness = FVariantManagerStyle::Get()->GetFloat("VariantManager.Spacings.BorderThickness");
+
+	return SNew( SBorder )
+	.BorderImage( FEditorStyle::GetBrush( "WhiteBrush" ) )
+	.BorderBackgroundColor( FVariantManagerStyle::Get()->GetColor( "VariantManager.Panels.ContentBackgroundColor" ) )
+	.Padding( FMargin( 0.0f, 0.0f, 0.0f, 2.0f ) )
+	[
+		// Actor column
+		SAssignNew( PropertiesSplitter, SSplitter )
+		.Orientation( Orient_Horizontal )
+		.PhysicalSplitterHandleSize( BorderThickness )
+		.HitDetectionSplitterHandleSize( BorderThickness )
+
+		+ SSplitter::Slot()
+		.Value( PropertiesColumnSizeData.LeftColumnWidth )
+		.OnSlotResized( SSplitter::FOnSlotResized::CreateLambda( []( float InNewWidth ) {} ) )
+		[
+			SNew(SVerticalBox)
+			+ SVerticalBox::Slot()
+			.MaxHeight(VM_COMMON_HEADER_MAX_HEIGHT)
+			.AutoHeight()
+			.Padding(FMargin(0.0f, VM_COMMON_PADDING, 0.0f, VM_COMMON_PADDING))
+			[
+				SNew(SBox)
+				.HeightOverride(VM_COMMON_HEADER_MAX_HEIGHT)
+				.HAlign(HAlign_Center)
+				.VAlign(VAlign_Center)
+				[
+					SNew(SHorizontalBox)
+
+					+SHorizontalBox::Slot()
+					.AutoWidth()
+					.VAlign(VAlign_Center)
+					[
+						SNew(STextBlock)
+						.TextStyle(FEditorStyle::Get(), "NormalText.Important")
+						.Text(LOCTEXT("ActorsText", "Actors"))
+					]
+
+					+SHorizontalBox::Slot()
+					.VAlign(VAlign_Center)
+					.HAlign(HAlign_Right)
+					.MaxWidth(24.0f)
+					.AutoWidth()
+					.Padding(FMargin(4.f, 0.f, 4.f, 0.f))
+					[
+						SNew(SButton)
+						.ButtonStyle(FEditorStyle::Get(), "NoBorder")
+						.ToolTipText(LOCTEXT("AddActorPlusTooltip", "Add a new actor binding to selected variants"))
+						.OnClicked(this, &SVariantManager::OnSummonAddActorMenu)
+						.ContentPadding(FMargin(2.0f, 1.0f))
+						.IsEnabled_Lambda([&VariantManagerPtr = VariantManagerPtr]() -> bool
+						{
+							if (TSharedPtr<FVariantManager> VariantManager = VariantManagerPtr.Pin())
+							{
+								TArray<UVariant*> Variants;
+								TArray<UVariantSet*> VariantSets;
+								VariantManager->GetSelection().GetSelectedVariantsAndVariantSets(Variants, VariantSets);
+
+								if (Variants.Num() > 0)
+								{
+									return true;
+								}
+
+								for (const UVariantSet* Set : VariantSets)
+								{
+									if (Set && Set->GetNumVariants() > 0)
+									{
+										return true;
+									}
+								}
+							}
+							return false;
+						})
+						.Content()
+						[
+							SNew(STextBlock)
+							.TextStyle(FEditorStyle::Get(), "NormalText.Important")
+							.Font(FEditorStyle::Get().GetFontStyle("FontAwesome.10"))
+							.Text(FEditorFontGlyphs::Plus)
+						]
+					]
+				]
+			]
+
+			+ SVerticalBox::Slot()
+			.Padding(FMargin(0.0f, 0.0f, 0.0f, 0.0f))
+			.VAlign(VAlign_Fill)
+			.FillHeight(1.0f)
+			[
+				SNew(SBox)
+				.MinDesiredHeight(30.0f)
+				[
+					ActorListView.ToSharedRef()
+				]
+			]
+		]
+
+		// Properties column
+		+ SSplitter::Slot()
+		.Value( TAttribute<float>::Create( TAttribute<float>::FGetter::CreateLambda( [this]()
+		{
+			return PropertiesNameColumnWidth + PropertiesValueColumnWidth;
+		} ) ) )
+		.OnSlotResized( SSplitter::FOnSlotResized::CreateLambda( [this]( float InNewWidth )
+		{
+			// InNewWidth is the new size of the property names + property values "combo", as they're a nested splitter
+			// We want to keep the property values column fixed in size though
+			PropertiesNameColumnWidth = FMath::Max(InNewWidth - PropertiesValueColumnWidth, 0.01f);
+		} ) )
+		[
+			SNew(SVerticalBox)
+			+ SVerticalBox::Slot()
+			.MaxHeight(VM_COMMON_HEADER_MAX_HEIGHT)
+			.AutoHeight()
+			.Padding(0.0f, VM_COMMON_PADDING, VM_COMMON_PADDING, VM_COMMON_PADDING)
+			[
+				// Headers
+				SNew(SSplitter)
+				.Orientation(Orient_Horizontal)
+				.PhysicalSplitterHandleSize( BorderThickness )
+				.HitDetectionSplitterHandleSize( BorderThickness )
+
+				+ SSplitter::Slot()
+				.Value( PropertiesColumnSizeData.MiddleColumnWidth )
+				.OnSlotResized( SSplitter::FOnSlotResized::CreateLambda( []( float InNewWidth ) {} ) )
+				[
+					SNew(SBox)
+					.HeightOverride(VM_COMMON_HEADER_MAX_HEIGHT)
+					.HAlign(HAlign_Center)
+					.VAlign(VAlign_Center)
+					[
+						SNew(SHorizontalBox)
+
+						+SHorizontalBox::Slot()
+						.AutoWidth()
+						.VAlign(VAlign_Center)
+						[
+							SNew(STextBlock)
+							.TextStyle(FEditorStyle::Get(), "NormalText.Important")
+							.Text(LOCTEXT("PropertiesText", "Properties"))
+						]
+
+						+SHorizontalBox::Slot()
+						.VAlign(VAlign_Center)
+						.HAlign(HAlign_Right)
+						.MaxWidth(24.0f)
+						.AutoWidth()
+						.Padding(FMargin(4.f, 0.f, 4.f, 0.f))
+						[
+							SNew(SButton)
+							.ButtonStyle(FEditorStyle::Get(), "NoBorder")
+							.ToolTipText(LOCTEXT("CapturePropertiesPlusTooltip", "Capture properties from the selected actor bindings"))
+							.OnClicked_Lambda([this]
+							{
+								CaptureNewPropertiesFromSelectedActors();
+								return FReply::Handled();
+							})
+							.ContentPadding(FMargin(2.0f, 1.0f))
+							.IsEnabled_Lambda([&VariantManagerPtr = VariantManagerPtr]() -> bool
+							{
+								if (TSharedPtr<FVariantManager> VariantManager = VariantManagerPtr.Pin())
+								{
+									return VariantManager->GetSelection().GetSelectedActorNodes().Num() > 0;
+								}
+								return false;
+							})
+							.Content()
+							[
+								SNew(STextBlock)
+								.TextStyle(FEditorStyle::Get(), "NormalText.Important")
+								.Font(FEditorStyle::Get().GetFontStyle("FontAwesome.10"))
+								.Text(FEditorFontGlyphs::Plus)
+							]
+						]
+					]
+				]
+				+ SSplitter::Slot()
+				.Value( PropertiesColumnSizeData.RightColumnWidth )
+				.OnSlotResized( PropertiesColumnSizeData.OnSecondSplitterChanged )
+				[
+					SNew(SBox)
+					.HeightOverride(VM_COMMON_HEADER_MAX_HEIGHT)
+					.HAlign(HAlign_Center)
+					.VAlign(VAlign_Center)
+					[
+						SNew(STextBlock)
+						.TextStyle(FEditorStyle::Get(), "NormalText.Important")
+						.Text(LOCTEXT("PropertiesValuesText", "Values"))
+					]
+				]
+			]
+			+ SVerticalBox::Slot()
+			.Padding(0.0f, 0.0f, VM_COMMON_PADDING, 0.0f)
+			.VAlign(VAlign_Fill)
+			.FillHeight(1.0f)
+			[
+				SAssignNew(CapturedPropertyListView, SListView<TSharedPtr<FVariantManagerPropertyNode>>)
+				.SelectionMode(ESelectionMode::Single)
+				.ListItemsSource(&DisplayedPropertyNodes)
+				.OnContextMenuOpening(this, &SVariantManager::OnPropertyListContextMenuOpening)
+				.OnGenerateRow(this, &SVariantManager::MakeCapturedPropertyRow)
+				.Visibility(EVisibility::Visible)
+			]
+		]
+	];
+}
+
+TSharedRef<SWidget> SVariantManager::GenerateRightTreeDependenciesRowContent()
+{
+	float BorderThickness = FVariantManagerStyle::Get()->GetFloat( "VariantManager.Spacings.BorderThickness" );
+
+	return SNew( SBorder )
+	.BorderImage( FEditorStyle::GetBrush( "WhiteBrush" ) )
+	.BorderBackgroundColor( FVariantManagerStyle::Get()->GetColor( "VariantManager.Panels.ContentBackgroundColor" ) )
+	.Padding( FMargin( BorderThickness, 0.0f, BorderThickness, 1.0f ) )
+	[
+		SNew(SVerticalBox)
+		.IsEnabled_Lambda([this]()
+		{
+			TSharedPtr<FVariantManager> PinnedVariantManager = VariantManagerPtr.Pin();
+			if (PinnedVariantManager.IsValid())
+			{
+				TArray<UVariant*> SelectedVariants;
+				TArray<UVariantSet*> SelectedVariantSets;
+				PinnedVariantManager->GetSelection().GetSelectedVariantsAndVariantSets(SelectedVariants, SelectedVariantSets);
+
+				return SelectedVariants.Num() == 1;
+			}
+			return false;
+		})
+
+		// Headers
+		+ SVerticalBox::Slot()
+		.MaxHeight(VM_COMMON_HEADER_MAX_HEIGHT)
+		.AutoHeight()
+		.Padding(0.0f, VM_COMMON_PADDING, VM_COMMON_PADDING, VM_COMMON_PADDING)
+		[
+			SNew(SSplitter)
+			.Orientation(Orient_Horizontal)
+			.PhysicalSplitterHandleSize( BorderThickness )
+			.HitDetectionSplitterHandleSize( BorderThickness )
+
+			+SSplitter::Slot()
+			.Value(DependenciesColumnSizeData.LeftColumnWidth)
+			.OnSlotResized(SSplitter::FOnSlotResized::CreateLambda([](float InNewWidth)
+			{
+				//This has to be bound or the splitter will take it upon itself to determine the size
+				//We do nothing here because it is handled by the column size data
+			}))
+			[
+				SNew(SBox)
+				.HeightOverride(VM_COMMON_HEADER_MAX_HEIGHT)
+				.HAlign(HAlign_Center)
+				.VAlign(VAlign_Center)
+				[
+					SNew( STextBlock )
+					.TextStyle( FEditorStyle::Get(), "NormalText.Important" )
+					.Text( LOCTEXT( "DependenciesVariantSetText", "Variant Set" ) )
+				]
+			]
+			+ SSplitter::Slot()
+			.Value( DependenciesColumnSizeData.MiddleColumnWidth )
+			.OnSlotResized( DependenciesColumnSizeData.OnFirstSplitterChanged )
+			[
+				SNew(SBox)
+				.HeightOverride(VM_COMMON_HEADER_MAX_HEIGHT)
+				.HAlign(HAlign_Center)
+				.VAlign(VAlign_Center)
+				.Padding(FMargin(0, 0, 0, 0))
+				[
+					SNew(STextBlock)
+					.TextStyle(FEditorStyle::Get(), "NormalText.Important")
+					.Text(LOCTEXT("DependenciesVariantText", "Variant"))
+				]
+			]
+			+ SSplitter::Slot()
+			.Value( DependenciesColumnSizeData.RightColumnWidth )
+			.OnSlotResized( DependenciesColumnSizeData.OnSecondSplitterChanged )
+			[
+				SNullWidget::NullWidget
+			]
+		]
+
+		// Dependencies
+		+ SVerticalBox::Slot()
+		.Padding(0.0f, 0.0f, VM_COMMON_PADDING, 1.0f)
+		.VAlign(VAlign_Top)
+		.AutoHeight()
+		[
+			SNew( SBox )
+			.MinDesiredHeight( 30.0f )
+			[
+				SAssignNew(DependenciesList, SListView<FVariantDependencyModelPtr>)
+				.SelectionMode(ESelectionMode::None)
+				.ListItemsSource(&DisplayedDependencies)
+				.OnGenerateRow(this, &SVariantManager::GenerateDependencyRow, true)
+				.Visibility(EVisibility::Visible)
+			]
+		]
+
+		// Separator
+		+ SVerticalBox::Slot()
+		.Padding(0.0f, 10.0f, VM_COMMON_PADDING, 10.0f)
+		.VAlign(VAlign_Top)
+		.MaxHeight(24)
+		[
+			SNew(SSeparator)
+			.SeparatorImage(FEditorStyle::GetBrush("ThinLine.Horizontal"))
+			.Thickness(4.f)
+			.Orientation(EOrientation::Orient_Horizontal)
+			.ColorAndOpacity(FLinearColor(0.1, 0.1, 0.1, 1.0))
+			.Visibility_Lambda([this]()
+			{
+				return DisplayedDependents.Num() > 0 ? EVisibility::Visible : EVisibility::Collapsed;
+			})
+		]
+
+		// Dependents
+		+ SVerticalBox::Slot()
+		.Padding(0.0f, 0.0f, VM_COMMON_PADDING, 1.0f)
+		.VAlign(VAlign_Top)
+		.AutoHeight()
+		[
+			SAssignNew(DependentsList, SListView<FVariantDependencyModelPtr>)
+			.SelectionMode(ESelectionMode::None)
+			.ListItemsSource(&DisplayedDependents)
+			.OnGenerateRow(this, &SVariantManager::GenerateDependencyRow, false)
+			.Visibility_Lambda( [ this ]()
+			{
+				return DisplayedDependents.Num() > 0 ? EVisibility::Visible : EVisibility::Collapsed;
+			})
+		]
+	];
+}
+
+TSharedRef<ITableRow> SVariantManager::GenerateDependencyRow(FVariantDependencyModelPtr Dependency, const TSharedRef<STableViewBase>& OwnerTable, bool bInteractionEnabled)
+{
+	TSharedPtr<FVariantManager> PinnedVariantManager = VariantManagerPtr.Pin();
+	if (!PinnedVariantManager.IsValid())
+	{
+		return SNew(STableRow<FVariantDependencyModelPtr>, OwnerTable);
+	}
+
+	TArray<UVariant*> Variants;
+	TArray<UVariantSet*> VariantSets;
+	PinnedVariantManager->GetSelection().GetSelectedVariantsAndVariantSets(Variants, VariantSets);
+	if (Variants.Num() != 1 || Variants[0] == nullptr )
+	{
+		return SNew(STableRow<FVariantDependencyModelPtr>, OwnerTable);
+	}
+
+	return SNew(SDependencyRow, OwnerTable, GetDependenciesColumnSizeData(), Dependency, bInteractionEnabled);
+}
+
 void SVariantManager::OnEditorSelectionChanged(UObject* NewSelection)
 {
 	if (!bRespondToEditorSelectionEvents)
@@ -3439,5 +3921,21 @@
 	RefreshActorList();
 }
 
+void SVariantManager::OnThumbnailChanged(UObject* VariantOrVariantSet)
+{
+	TSharedPtr<FVariantManager> PinnedVariantManager = VariantManagerPtr.Pin();
+	if (!PinnedVariantManager.IsValid() || !VariantOrVariantSet)
+	{
+		return;
+	}
+
+	if (VariantOrVariantSet->GetOutermost() == PinnedVariantManager->GetCurrentLevelVariantSets()->GetOutermost())
+	{
+		RefreshVariantTree();
+	}
+}
+
+#undef VM_COMMON_PADDING
+#undef VM_COMMON_HEADER_MAX_HEIGHT
 #undef LOCTEXT_NAMESPACE
 
diff --git a/Engine/Plugins/Enterprise/VariantManager/Source/VariantManager/Private/SVariantManager.h b/Engine/Plugins/Enterprise/VariantManager/Source/VariantManager/Private/SVariantManager.h
--- a/Engine/Plugins/Enterprise/VariantManager/Source/VariantManager/Private/SVariantManager.h
+++ b/Engine/Plugins/Enterprise/VariantManager/Source/VariantManager/Private/SVariantManager.h
@@ -14,41 +14,23 @@
 #include "VariantManager.h"
 
 class FExtender;
+class FTransactionObjectEvent;
 class FVariantManager;
-class UVariant;
 class FVariantManagerDisplayNode;
-class SVariantManagerNodeTreeView;
 class FVariantManagerDisplayNode;
-class SVariantManagerActorListView;
 class FVariantManagerPropertyNameNode;
 class FVariantManagerPropertyNode;
 class ITableRow;
-class STableViewBase;
 class SSplitter;
-class FTransactionObjectEvent;
-struct FSlateImageBrush;
+class STableViewBase;
+class SVariantManagerActorListView;
+class SVariantManagerNodeTreeView;
+class UVariant;
 enum class EMapChangeType : uint8;
-
-namespace VariantManagerLayoutConstants
-{
-	/** The amount to indent child nodes of the layout tree */
-	const float IndentAmount = 10.0f;
-
-	/** Height of each folder node */
-	const float FolderNodeHeight = 20.0f;
-
-	/** Height of each object node */
-	const float ObjectNodeHeight = 20.0f;
-
-	/** Height of each section area if there are no sections (note: section areas may be larger than this if they have children. This is the height of a section area with no children or all children hidden) */
-	const float SectionAreaDefaultHeight = 15.0f;
-
-	/** Height of each key area */
-	const float KeyAreaHeight = 15.0f;
-
-	/** Height of each category node */
-	const float CategoryNodeHeight = 15.0f;
-}
+struct FSlateImageBrush;
+struct FVariantDependency;
+struct FVariantDependencyModel;
+using FVariantDependencyModelPtr = TSharedPtr<FVariantDependencyModel>;
 
 // Convenience struct to save/load how the user configured the main splitters
 struct FSplitterValues
@@ -56,13 +38,24 @@ struct FSplitterValues
 	float VariantColumn = 0.25f;
 	float ActorColumn = 0.25f;
 	float PropertyNameColumn = 0.25f;
-	float PropertyValueColumn = 0.25f;
+	float DependenciesVariantSetsColumn = 0.33f;
+	float DependenciesVariantColumn = 0.33f;
 
 	FSplitterValues(){};
 	FSplitterValues(FString& InSerialized);
 	FString ToString();
 };
 
+struct FColumnSizeData
+{
+	TAttribute<float> LeftColumnWidth;
+	TAttribute<float> MiddleColumnWidth;
+	TAttribute<float> RightColumnWidth;
+	// There are three columns, but only two splitters
+	SSplitter::FOnSlotResized OnFirstSplitterChanged;
+	SSplitter::FOnSlotResized OnSecondSplitterChanged;
+};
+
 class SVariantManager
 	: public SCompoundWidget
 	, public FNotifyHook
@@ -177,20 +170,27 @@ public:
 	void SortDisplayNodes(TArray<TSharedRef<FVariantManagerDisplayNode>>& DisplayNodes);
 
 	TSharedRef<SWidget> MakeAddButton();
-	FDetailColumnSizeData& GetPropertyColumnSizeData()
+	FColumnSizeData& GetPropertiesColumnSizeData()
 	{
-		return ColumnSizeData;
+		return PropertiesColumnSizeData;
+	}
+	FColumnSizeData& GetDependenciesColumnSizeData()
+	{
+		return DependenciesColumnSizeData;
 	}
 
 	TSharedRef<ITableRow> MakeCapturedPropertyRow(TSharedPtr<FVariantManagerPropertyNode> Item, const TSharedRef<STableViewBase>& OwnerTable);
 	TSharedPtr<SWidget> OnPropertyListContextMenuOpening();
 
+	void OnOutlinerNodeSelectionChanged();
 	void OnActorNodeSelectionChanged();
+	void OnVariantDependenciesUpdated(UVariant* ParentVariant);
 
 	// These completely refresh the data and the view for each display
 	void RefreshVariantTree();
 	void RefreshActorList();
 	void RefreshPropertyList();
+	void RefreshDependencyLists();
 	void UpdatePropertyDefaults();
 
 	void OnBlueprintCompiled();
@@ -212,6 +212,21 @@ public:
 
 	FReply OnAddVariantSetClicked();
 	FReply OnSummonAddActorMenu();
+	FReply OnAddDependencyClicked();
+
+	// Callbacks for property ColumnSizeData
+	float OnGetPropertiesActorColumnWidth() const { return 1.0f - PropertiesNameColumnWidth - PropertiesValueColumnWidth; }
+	float OnGetPropertiesNameColumnWidth() const { return PropertiesNameColumnWidth; }
+	float OnGetPropertiesValueColumnWidth() const { return PropertiesValueColumnWidth; }
+	void OnSetPropertiesNameColumnWidth( float InWidth ) { PropertiesNameColumnWidth = InWidth; }
+	void OnSetPropertiesValueColumnWidth( float InWidth ) { PropertiesValueColumnWidth = InWidth; }
+
+	// Callbacks for dependencies ColumnSizeData
+	float OnGetDependenciesVariantSetColumnWidth() const { return 1.0f - DependenciesVariantColumnWidth - DependenciesControlColumnWidth; }
+	float OnGetDependenciesVariantColumnWidth() const { return DependenciesVariantColumnWidth; }
+	float OnGetDependenciesControlColumnWidth() const { return DependenciesControlColumnWidth; }
+	void OnSetDependenciesVariantColumnWidth( float InWidth ) { DependenciesVariantColumnWidth = InWidth; }
+	void OnSetDependenciesControlColumnWidth( float InWidth ) { DependenciesControlColumnWidth = InWidth; }
 
 	void OnObjectTransacted(UObject* Object, const class FTransactionObjectEvent& Event);
 	void OnObjectPropertyChanged(UObject* Object, struct FPropertyChangedEvent& Event);
@@ -219,8 +234,25 @@ public:
 	void OnPieEvent(bool bIsSimulating);
 	void OnEditorSelectionChanged(UObject* NewSelection);
 
+	void OnThumbnailChanged(UObject* VariantOrVariantSet);
+
 	void ReorderPropertyNodes(const TArray<TSharedPtr<FVariantManagerPropertyNode>>& TheseNodes, TSharedPtr<FVariantManagerPropertyNode> Pivot, EItemDropZone RelativePosition);
 
+private:
+	enum class ERightTreeRowType
+	{
+		PropertiesHeader,
+		PropertiesContent,
+		DependenciesHeader,
+		DependenciesContent
+	};
+
+	TSharedRef<ITableRow> GenerateRightTreeRow( TSharedRef<ERightTreeRowType> RowType, const TSharedRef<STableViewBase>& OwnerTable );
+	TSharedRef<SWidget> GenerateRightTreeHeaderRowContent( ERightTreeRowType RowType, TSharedRef<STableRow<TSharedRef<ERightTreeRowType>>> InTableRow );
+	TSharedRef<SWidget> GenerateRightTreePropertiesRowContent();
+	TSharedRef<SWidget> GenerateRightTreeDependenciesRowContent();
+	TSharedRef<ITableRow> GenerateDependencyRow( FVariantDependencyModelPtr Dependency, const TSharedRef<STableViewBase>& OwnerTable, bool bInteractionEnabled );
+
 private:
 	TWeakPtr<FVariantManager> VariantManagerPtr;
 
@@ -232,6 +264,12 @@ private:
 	TSharedPtr<SListView<TSharedPtr<FVariantManagerPropertyNode>>> CapturedPropertyListView;
 	TArray<TSharedPtr<FVariantManagerPropertyNode>> DisplayedPropertyNodes;
 
+	TSharedPtr<SListView<FVariantDependencyModelPtr>> DependenciesList;
+	TArray<FVariantDependencyModelPtr> DisplayedDependencies;
+
+	TSharedPtr<SListView<FVariantDependencyModelPtr>> DependentsList;
+	TArray<FVariantDependencyModelPtr> DisplayedDependents;
+
 	// We use paths here to avoid having to check if the bindings are resolved
 	TSet<FString> CachedSelectedActorPaths;
 	TSet<FString> CachedDisplayedActorPaths;
@@ -247,9 +285,15 @@ private:
 
 	bool bAutoCaptureProperties = false;
 
-	// Mirrors detailview, its used by all splitters in the column, so that they move in sync
-	FDetailColumnSizeData ColumnSizeData;
-	float RightPropertyColumnWidth;
+	FColumnSizeData PropertiesColumnSizeData;
+	float PropertiesNameColumnWidth;
+	float PropertiesValueColumnWidth;
+
+	FColumnSizeData DependenciesColumnSizeData;
+	float DependenciesVariantColumnWidth;
+	float DependenciesControlColumnWidth;
+
+	FSplitterValues SplitterValues;
 
 	FDelegateHandle OnObjectTransactedHandle;
 	FDelegateHandle OnBlueprintCompiledHandle;
@@ -259,9 +303,15 @@ private:
 	FDelegateHandle OnBeginPieHandle;
 	FDelegateHandle OnEndPieHandle;
 	FDelegateHandle OnEditorSelectionChangedHandle;
+	FDelegateHandle OnVariantThumbnailUpdatedHandle;
+	FDelegateHandle OnVariantSetThumbnailUpdatedHandle;
 
 	// We keep track of this to remember splitter values between loads
 	TSharedPtr<SSplitter> MainSplitter;
+	TSharedPtr<SSplitter> PropertiesSplitter;
+	TSharedPtr<SSplitter> DependenciesSplitter;
+
+	TArray<TSharedRef<ERightTreeRowType>> RightTreeRootItems;
 
 	// TODO: Make separate VariantManagerStyle
 	TSharedPtr<FSlateImageBrush> RecordButtonBrush;
diff --git a/Engine/Shaders/Private/RayTracing/RayTracingFinalGatherRGS.usf b/Engine/Shaders/Private/RayTracing/RayTracingFinalGatherRGS.usf
--- a/Engine/Shaders/Private/RayTracing/RayTracingFinalGatherRGS.usf
+++ b/Engine/Shaders/Private/RayTracing/RayTracingFinalGatherRGS.usf
@@ -17,6 +17,11 @@
 #include "../QuasiRandom.ush"
 
 #define USE_PATHTRACING_MATERIALS 0
+#define USE_REJECTION_CRITERIA 0
+#define USE_REPROJECTION 1
+#define USE_NEIGHBOR_SAMPLES 1
+#define USE_TEMPORAL_SAMPLES 1
+
 
 uint SampleIndex;
 uint SamplesPerPixel;
@@ -26,6 +31,10 @@ uint RenderTileOffsetY;
 float DiffuseThreshold;
 float MaxNormalBias;
 float FinalGatherDistance;
+uint GatherFilterWidth;
+uint GatherPointIterations;
+float DepthRejectionKernel;
+float NormalRejectionKernel;
 
 RaytracingAccelerationStructure TLAS;
 
@@ -33,18 +42,34 @@ int2 GatherPointsResolution;
 StructuredBuffer<FGatherPoints> GatherPointsBuffer;
 
 RWTexture2D<float4> RWGlobalIlluminationUAV;
-RWTexture2D<float> RWRayDistanceUAV;
+RWTexture2D<float2> RWGlobalIlluminationRayDistanceUAV;
 
 uint DispatchThreadIdToLinearIndex(uint2 DispatchThreadId)
 {
 	return DispatchThreadId.y * GatherPointsResolution.x + DispatchThreadId.x;
 }
 
+bool ShouldBilateralReject(float2 PixelCoord, float Z, float LocalZ, float3 WorldNormal, float3 LocalWorldNormal)
+{
+	// Gather points dimension guard
+	if (any(PixelCoord < 0) || any(PixelCoord > View.BufferSizeAndInvSize.xy - int2(1, 1))) return true;
+
+	// Depth deviation check
+	float DepthDeviation = abs(Z - LocalZ);
+	if (DepthDeviation > DepthRejectionKernel) return true;
+
+	// Normal deviation check
+	float NormalDeviation = 1.0 - saturate(dot(WorldNormal, LocalWorldNormal));
+	if (NormalDeviation > NormalRejectionKernel) return true;
+
+	return false;
+}
+
 RAY_TRACING_ENTRY_RAYGEN(RayTracingFinalGatherRGS)
 {
 	uint2 DispatchThreadId = DispatchRaysIndex().xy + View.ViewRectMin + uint2(RenderTileOffsetX, RenderTileOffsetY);
 	RWGlobalIlluminationUAV[DispatchThreadId] = 0.0;
-	RWRayDistanceUAV[DispatchThreadId] = float2(-1.0, 0.0);
+	RWGlobalIlluminationRayDistanceUAV[DispatchThreadId] = float2(-1.0, 0.0);
 
 	// Get G-Buffer surface data
 	uint2 PixelCoord = GetPixelCoord(DispatchThreadId, UpscaleFactor);
@@ -83,16 +108,16 @@ RAY_TRACING_ENTRY_RAYGEN(RayTracingFinalGatherRGS)
 	}
 
 	// Construct position, normal, and view direction
-	float DeviceZ = SceneDepthBuffer.Load(int3(PixelCoord, 0)).r;
+	float DeviceZ = SceneDepthTexture.Load(int3(PixelCoord, 0)).r;
 	float3 WorldPosition;
 	float3 CameraDirection;
 	ReconstructWorldPositionAndCameraDirectionFromDeviceZ(PixelCoord, DeviceZ, WorldPosition, CameraDirection);
 	float3 WorldNormal = GBufferData.WorldNormal;
 
-	uint GatherPointsIndex = DispatchThreadIdToLinearIndex(DispatchThreadId);
-
+//#if USE_REJECTION_CRITERIA
 	// Define rejection criteria and evaluate
 	bool bShouldReject[MAXIMUM_GATHER_POINTS_PER_PIXEL];
+//#endif
 
 	// Create ray differential based on rejection distance (in pixels)
 	UV = (float2(PixelCoord) + 0.5 + float2(0.0, FinalGatherDistance * UpscaleFactor)) * InvBufferSize;
@@ -104,23 +129,74 @@ RAY_TRACING_ENTRY_RAYGEN(RayTracingFinalGatherRGS)
 	float3 WorldDeltaPosition = Ray.Origin + T * Ray.Direction;
 	float RejectionDistance = length(WorldDeltaPosition - WorldPosition);
 
-	FRejectionCriteria RejectionCriteria = CreateRejectionCriteria(WorldPosition, RejectionDistance);
-	CreateRejectionMask(RejectionCriteria, GatherPointsBuffer, GatherPointsIndex, 0, GatherPointsResolution.x * GatherPointsResolution.y, SamplesPerPixel, bShouldReject);
-
 	float3 ExitantRadiance = 0.0;
-	float HitDistance = 1.0e7;
+	float3 MaxExitantRadianceSample = 0.0;
+	//float HitDistance = 1.0e7;
+	float HitDistance = 0.0;
 	float HitCount = 0.0;
 	uint ValidSamplesPerPixel = 0;
+
+	// Gather from temporal history
+#if USE_TEMPORAL_SAMPLES
+
+	uint GatherPointsIndex = DispatchThreadIdToLinearIndex(DispatchThreadId);
+	FRejectionCriteria RejectionCriteria = CreateRejectionCriteria(WorldPosition, RejectionDistance);
+#if USE_REJECTION_CRITERIA
+	CreateRejectionMask(RejectionCriteria, GatherPointsBuffer, GatherPointsIndex, 0, GatherPointsResolution.x* GatherPointsResolution.y, SamplesPerPixel, bShouldReject);
+#endif
+
 	for (uint LocalSampleIndex = 0; LocalSampleIndex < SamplesPerPixel; ++LocalSampleIndex)
 	{
+		//bShouldReject[LocalSampleIndex] = true;
+		//if (LocalSampleIndex == SampleIndex) continue;
+#if USE_REJECTION_CRITERIA
+		if (bShouldReject[LocalSampleIndex])
+		{
+			bShouldReject[LocalSampleIndex] = true;
+			//continue;
+		}
+#endif
+
+#if USE_REPROJECTION
+		float4 ScreenPositionTime = mul(float4(WorldPosition, 1.0), GatherPointData.ViewMatrices[LocalSampleIndex]);
+		ScreenPositionTime.xyz /= ScreenPositionTime.w;
+		if (any(abs(ScreenPositionTime.xy) > 1.0))
+		{
+			bShouldReject[LocalSampleIndex] = true;
+			//continue;
+		}
+
+		float2 ScreenUVTime = ScreenPositionTime.xy * float2(0.5, -0.5) + 0.5;
+		GatherPointsIndex = DispatchThreadIdToLinearIndex(int2(ScreenUVTime * GatherPointsResolution + View.ViewRectMin));
+#endif
+
+		FGatherSample GatherSample = ReadGatherSample(GatherPointsBuffer, GatherPointsIndex, LocalSampleIndex, GatherPointsResolution.x * GatherPointsResolution.y);
+
+#if USE_REPROJECTION
+		float3 GPDistance2 = RejectionCriteria.CreationPoint - GatherSample.CreationPoint;
+		if (dot(GPDistance2, GPDistance2) > RejectionCriteria.DistanceThreshold2)
+		{
+			bShouldReject[LocalSampleIndex] = true;
+		}
+#endif
+
+		uint SampleIndexEnd = (SampleIndex + GatherPointIterations) % SamplesPerPixel;
+		if (SampleIndex < SampleIndexEnd)
+		{
+			if (LocalSampleIndex >= SampleIndex && LocalSampleIndex < SampleIndexEnd)
+				bShouldReject[LocalSampleIndex] = false;
+		}
+		else
+		{
+			if (LocalSampleIndex < SampleIndexEnd || LocalSampleIndex >= SampleIndex)
+				bShouldReject[LocalSampleIndex] = false;
+		}
 
 		if (bShouldReject[LocalSampleIndex])
 		{
 			continue;
 		}
-		ValidSamplesPerPixel++;
 
-		FGatherSample GatherSample = ReadGatherSample(GatherPointsBuffer, GatherPointsIndex, LocalSampleIndex, GatherPointsResolution.x * GatherPointsResolution.y);
 		float3 GPPosition = GatherSample.Position;
 		float3 GPIrradiance = GatherSample.Irradiance;
 
@@ -128,17 +204,28 @@ RAY_TRACING_ENTRY_RAYGEN(RayTracingFinalGatherRGS)
 		half3 N = WorldNormal;
 		half3 V = -CameraDirection;
 		half3 L = normalize(GPPosition - WorldPosition);
-		float NoL = saturate(dot(N, L));
+		float NoL = dot(N, L);
+		if (NoL <= 0.0) continue;
+
 		FShadowTerms ShadowTerms = { 0.0, 0.0, 0.0, InitHairTransmittanceData() };
 		FDirectLighting LightingSample = EvaluateBxDF(GBufferData, N, V, L, NoL, ShadowTerms);
 		float3 RayThroughput = LightingSample.Diffuse / DiffuseColor;
-		float RayPdf = NoL / PI;
+
+		float TargetPdf = NoL / PI;
+		//float TargetPdf = 1.0 / (4.0 * PI);
+		float SourcePdf = GatherSample.Pdf;
+		float Weight = TargetPdf / SourcePdf;
+		float RayPdf = TargetPdf;
+
+		// Firefly prevention attempt
+		if (RayPdf < 1.0e-3) continue;
+		ValidSamplesPerPixel++;
 
 		// Create ray
 		RayDesc Ray;
 		Ray.Origin = WorldPosition;
 		Ray.Direction = normalize(GPPosition - WorldPosition);
-		Ray.TMin = 0.01;
+		Ray.TMin = 0.05;
 		float RayLength = length(GPPosition - WorldPosition);
 		Ray.TMax = max(RayLength - 0.01, Ray.TMin);
 
@@ -160,28 +247,166 @@ RAY_TRACING_ENTRY_RAYGEN(RayTracingFinalGatherRGS)
 		// No hit indicates successful next-event connection
 		if (MinimalPayload.IsMiss() && RayPdf > 0.0)
 		{
-			ExitantRadiance += GPIrradiance * RayThroughput / RayPdf;
-			HitDistance = min(HitDistance, Ray.TMax);
+			float3 ExitantRadianceSample = GPIrradiance * RayThroughput * Weight / RayPdf;
+			// Firefly reduction
+			//ExitantRadianceSample *= 1.0 / (1.0 + Luminance(ExitantRadianceSample));
+			ExitantRadiance += ExitantRadianceSample;
+			MaxExitantRadianceSample = Luminance(ExitantRadianceSample) > Luminance(MaxExitantRadianceSample) ? ExitantRadianceSample : MaxExitantRadianceSample;
+			HitDistance += rcp(Ray.TMax);
+
 			HitCount += 1.0;
 		}
+		else
+		{
+			HitDistance += rcp(MinimalPayload.HitT);
+		}
+	}
+#endif
+
+	int GatherFilterWidth2 = GatherFilterWidth * GatherFilterWidth;
+#if USE_NEIGHBOR_SAMPLES
+	// Gather from local neighborhood if confidence in temporal history is low
+	{
+		// dxr_todo: Iterate backwards until history rejection declares all remaining samples invalid
+		for (uint ParentSampleIndex = 0; ParentSampleIndex < SamplesPerPixel; ++ParentSampleIndex)
+		{
+			if (bShouldReject[ParentSampleIndex]) continue;
+
+#if USE_REPROJECTION
+			float4 ScreenPositionTime = mul(float4(WorldPosition, 1.0), GatherPointData.ViewMatrices[ParentSampleIndex]);
+			ScreenPositionTime.xyz /= ScreenPositionTime.w;
+			float2 ScreenUVTime = ScreenPositionTime.xy * float2(0.5, -0.5) + 0.5;
+			int2 GatherPointCentroid = int2(ScreenUVTime * GatherPointsResolution + View.ViewRectMin);
+#endif
+
+			for (uint LocalSampleIndex0 = 0; LocalSampleIndex0 < GatherFilterWidth2; ++LocalSampleIndex0)
+			{
+				uint Midpoint = (GatherFilterWidth2 - 1) / 2;
+				if (LocalSampleIndex0 == Midpoint) continue;
+
+				int2 LocalIndex = int2(LocalSampleIndex0 % GatherFilterWidth, LocalSampleIndex0 / GatherFilterWidth);
+
+#if 0
+				// Preserve original stratification domain
+				//int2 GlobalIndex = DispatchThreadId - (DispatchThreadId % GatherFilterWidth) + LocalIndex;
+				int2 GlobalIndex = GatherPointCentroid - (GatherPointCentroid % GatherFilterWidth) + LocalIndex;
+#else
+				// Ignore stratifaction domain and respect local neighborhood instead
+				int2 Midpoint2 = GatherFilterWidth / 2;
+				LocalIndex -= Midpoint2;
+				int2 GlobalIndex = DispatchThreadId + LocalIndex;
+				if (any(GlobalIndex < 0) || any(GlobalIndex > View.BufferSizeAndInvSize.xy - 1)) continue;
+#endif
+				// Apply temporal bilateral rejection criteria to prevent haloing around occluding objects
+				uint2 LocalPixelCoord = GetPixelCoord(GlobalIndex, UpscaleFactor);
+				float LocalDeviceZ = SceneDepthTexture.Load(int3(LocalPixelCoord, 0)).r;
+				float3 LocalWorldNormal = DecodeNormal(GBufferATexture.Load(int3(LocalPixelCoord, 0)).xyz);
+				if (ShouldBilateralReject(LocalPixelCoord, DeviceZ, LocalDeviceZ, WorldNormal, LocalWorldNormal)) continue;
+
+				uint GatherPointsIndex = DispatchThreadIdToLinearIndex(GlobalIndex);
+				FGatherSample GatherSample = ReadGatherSample(GatherPointsBuffer, GatherPointsIndex, ParentSampleIndex, GatherPointsResolution.x * GatherPointsResolution.y);
+				float3 GPPosition = GatherSample.Position;
+				float3 GPIrradiance = GatherSample.Irradiance;
+
+				// Initialize throughput based on BRDF
+				half3 N = WorldNormal;
+				half3 V = -CameraDirection;
+				half3 L = normalize(GPPosition - WorldPosition);
+				float NoL = dot(N, L);
+
+				// Threshold against grazing angle, so the the pdf's don't explode
+				if (NoL < 1.0e-3) continue;
+
+				FShadowTerms ShadowTerms = { 0.0, 0.0, 0.0, InitHairTransmittanceData() };
+				FDirectLighting LightingSample = EvaluateBxDF(GBufferData, N, V, L, NoL, ShadowTerms);
+				float3 RayThroughput = LightingSample.Diffuse / DiffuseColor;
+				float TargetPdf = NoL / PI;
+				//float TargetPdf = 1.0 / (4.0 * PI);
+				float SourcePdf = GatherSample.Pdf;
+				float Weight = TargetPdf / SourcePdf;
+				float RayPdf = TargetPdf;
+
+				ValidSamplesPerPixel++;
+#if USE_NEIGHBOR_VISIBILITY_TEST
+				// Create ray
+				RayDesc Ray;
+				Ray.Origin = WorldPosition;
+				Ray.Direction = normalize(GPPosition - WorldPosition);
+				Ray.TMin = 0.01;
+				float RayLength = length(GPPosition - WorldPosition);
+				Ray.TMax = max(RayLength - 0.01, Ray.TMin);
+
+				// Trace ray
+				uint RayFlags = RAY_FLAG_ACCEPT_FIRST_HIT_AND_END_SEARCH | RAY_FLAG_SKIP_CLOSEST_HIT_SHADER;
+				const uint InstanceInclusionMask = RAY_TRACING_MASK_OPAQUE;
+
+#if !ENABLE_TWO_SIDED_GEOMETRY
+				RayFlags |= RAY_FLAG_CULL_BACK_FACING_TRIANGLES;
+#endif
+
+				FMinimalPayload MinimalPayload = TraceVisibilityRay(
+					TLAS,
+					RayFlags,
+					InstanceInclusionMask,
+					PixelCoord,
+					Ray);
+
+				// No hit indicates successful next-event connection
+				if (MinimalPayload.IsMiss() && RayPdf > 0.0)
+#endif
+				{
+					float3 ExitantRadianceSample = GPIrradiance * RayThroughput * Weight / RayPdf;
+					// Firefly reduction
+					//ExitantRadianceSample *= 1.0 / (1.0 + Luminance(ExitantRadianceSample));
+					ExitantRadiance += ExitantRadianceSample;
+					MaxExitantRadianceSample = Luminance(ExitantRadianceSample) > Luminance(MaxExitantRadianceSample) ? ExitantRadianceSample : MaxExitantRadianceSample;
+					HitDistance += rcp(Ray.TMax);
+					HitCount += 1.0;
+				}
+#if USE_NEIGHBOR_VISIBILITY_TEST
+				else
+				{
+					HitDistance += rcp(MinimalPayload.HitT);
+				}
+#endif
+			}
+		}
 	}
+#endif
 
 	float AmbientOcclusion = 0.0;
+#if 1
+	// Remove worst outlier when using neighbor sampling
+	if (GatherFilterWidth > 1)
+	{
+		ExitantRadiance -= MaxExitantRadianceSample;
+		ValidSamplesPerPixel--;
+	}
+#endif
+
 	if (ValidSamplesPerPixel > 0)
 	{
 		ExitantRadiance /= ValidSamplesPerPixel;
 		AmbientOcclusion = HitCount / ValidSamplesPerPixel;
+		HitDistance = ValidSamplesPerPixel / HitDistance;
 	}
 #if USE_PREEXPOSURE
 	ExitantRadiance *= View.PreExposure;
 #endif
 	ExitantRadiance = ClampToHalfFloatRange(ExitantRadiance);
-
+#if 0
+	// Threshold against low luminance..
+	float LuminanceThreshold = 1.0 / 256.0;
+	if (Luminance(ExitantRadiance * DiffuseColor) < LuminanceThreshold)
+	{
+		ExitantRadiance = 0.0;
+	}
+#endif
 	if (HitCount == 0.0)
 	{
 		HitDistance = -1.0;
 	}
 
 	RWGlobalIlluminationUAV[DispatchThreadId] = float4(ExitantRadiance, AmbientOcclusion);
-	RWRayDistanceUAV[DispatchThreadId] = float2(HitDistance, ValidSamplesPerPixel);
+	RWGlobalIlluminationRayDistanceUAV[DispatchThreadId] = float2(HitDistance, ValidSamplesPerPixel);
 }
diff --git a/Engine/Source/Runtime/D3D12RHI/Private/D3D12CommandContext.cpp b/Engine/Source/Runtime/D3D12RHI/Private/D3D12CommandContext.cpp
--- a/Engine/Source/Runtime/D3D12RHI/Private/D3D12CommandContext.cpp
+++ b/Engine/Source/Runtime/D3D12RHI/Private/D3D12CommandContext.cpp
@@ -45,6 +45,46 @@ FD3D12CommandContextBase::FD3D12CommandContextBase(class FD3D12Adapter* InParent
 {
 }
 
+static D3D12_RESOURCE_STATES GetValidResourceStates(D3D12_COMMAND_LIST_TYPE CommandListType)
+{
+	// For reasons, we can't just list the allowed states, we have to list the disallowed states.
+	// For reference on allowed/disallowed states, see:
+	//    https://microsoft.github.io/DirectX-Specs/d3d/CPUEfficiency.html#state-support-by-command-list-type
+	
+	const D3D12_RESOURCE_STATES DisallowedDirectStates =
+		static_cast<D3D12_RESOURCE_STATES>(0);
+
+	const D3D12_RESOURCE_STATES DisallowedComputeStates =
+		DisallowedDirectStates |
+		D3D12_RESOURCE_STATE_VERTEX_AND_CONSTANT_BUFFER |
+		D3D12_RESOURCE_STATE_INDEX_BUFFER |
+		D3D12_RESOURCE_STATE_RENDER_TARGET |
+		D3D12_RESOURCE_STATE_DEPTH_WRITE |
+		D3D12_RESOURCE_STATE_DEPTH_READ |
+		D3D12_RESOURCE_STATE_PIXEL_SHADER_RESOURCE |
+		D3D12_RESOURCE_STATE_STREAM_OUT |
+		D3D12_RESOURCE_STATE_INDIRECT_ARGUMENT |
+		D3D12_RESOURCE_STATE_RESOLVE_DEST |
+		D3D12_RESOURCE_STATE_RESOLVE_SOURCE;
+
+	const D3D12_RESOURCE_STATES DisallowedCopyStates =
+		DisallowedComputeStates |
+		D3D12_RESOURCE_STATE_UNORDERED_ACCESS |
+		D3D12_RESOURCE_STATE_NON_PIXEL_SHADER_RESOURCE;
+
+
+	if (CommandListType == D3D12_COMMAND_LIST_TYPE_COPY)
+	{
+		return ~DisallowedCopyStates;
+	}
+
+	if (CommandListType == D3D12_COMMAND_LIST_TYPE_COMPUTE)
+	{
+		return ~DisallowedComputeStates;
+	}
+
+	return ~DisallowedDirectStates;
+}
 
 FD3D12CommandContext::FD3D12CommandContext(FD3D12Device* InParent, bool InIsDefaultContext, bool InIsAsyncComputeContext) :
 	FD3D12CommandContextBase(InParent->GetParentAdapter(), InParent->GetGPUMask(), InIsDefaultContext, InIsAsyncComputeContext),
@@ -68,6 +108,7 @@ FD3D12CommandContext::FD3D12CommandContext(FD3D12Device* InParent, bool InIsDefa
 	VRSShadingRate(D3D12_SHADING_RATE_1X1),
 #endif
 	SkipFastClearEliminateState(D3D12_RESOURCE_STATE_PIXEL_SHADER_RESOURCE | D3D12_RESOURCE_STATE_NON_PIXEL_SHADER_RESOURCE),
+	ValidResourceStates(GetValidResourceStates(InIsAsyncComputeContext ? D3D12_COMMAND_LIST_TYPE_COMPUTE : D3D12_COMMAND_LIST_TYPE_DIRECT)),
 #if PLATFORM_SUPPORTS_VIRTUAL_TEXTURES
 	bNeedFlushTextureCache(false),
 #endif
@@ -176,7 +217,7 @@ void FD3D12CommandContext::RHIPushEvent(const TCHAR* Name, FColor Color)
 
 #if NV_AFTERMATH
 		// Only track aftermath for default context?
-		if (IsDefaultContext() && GDX12NVAfterMathEnabled)
+		if (IsDefaultContext() && GDX12NVAfterMathEnabled && GDX12NVAfterMathMarkers)
 			GFSDK_Aftermath_SetEventMarker(CommandListHandle.AftermathCommandContext(), &GPUEventStack[0], GPUEventStack.Num() * sizeof(uint32));
 #endif // NV_AFTERMATH		
 	}
@@ -233,16 +274,6 @@ void FD3D12CommandContext::RHIPopEvent()
 #endif
 }
 
-void FD3D12CommandContext::RHIAutomaticCacheFlushAfterComputeShader(bool bEnable)
-{
-	StateCache.AutoFlushComputeShaderCache(bEnable);
-}
-
-void FD3D12CommandContext::RHIFlushComputeShaderCache()
-{
-	StateCache.FlushComputeShaderCache(true);
-}
-
 FD3D12CommandListManager& FD3D12CommandContext::GetCommandListManager()
 {
 	return bIsAsyncComputeContext ? GetParentDevice()->GetAsyncCommandListManager() : GetParentDevice()->GetCommandListManager();
@@ -322,6 +353,7 @@ FD3D12CommandListHandle FD3D12CommandContext::FlushCommands(bool WaitForCompleti
 #endif
 
 	FD3D12Device* Device = GetParentDevice();
+	const bool bIsCommandListOpen = !CommandListHandle.IsClosed();
 	const bool bHasPendingWork = Device->PendingCommandLists.Num() > 0;
 	const bool bHasDoneWork = HasDoneWork() || bHasPendingWork;
 	const bool bOpenNewCmdList = WaitForCompletion || bHasDoneWork || bHasProfileGPUAction;
@@ -330,24 +362,45 @@ FD3D12CommandListHandle FD3D12CommandContext::FlushCommands(bool WaitForCompleti
 	if (bOpenNewCmdList)
 	{
 		// Close the current command list
-		CloseCommandList();
+		if (bIsCommandListOpen)
+		{
+			CloseCommandList();
+		}
 
 		if (bHasPendingWork)
 		{
-			// Submit all pending command lists and the current command list
-			Device->PendingCommandLists.Add(CommandListHandle);
+			// Submit all pending command lists and the current command list if it was still open
+			if (bIsCommandListOpen)
+			{
+				Device->PendingCommandLists.Add(CommandListHandle);				
+			}
+			else
+			{
+				// otherwise it should be already part of the pending list
+				check(Device->PendingCommandLists.Contains(CommandListHandle));
+
+				// This use case should only happen when force flush is called from the rendering thread using the FScopedRHIThreadStaller
+				// which could insert a flush in between pending command lists
+				check(IsInRenderingThread());
+			}
 			GetCommandListManager().ExecuteCommandLists(Device->PendingCommandLists, WaitForCompletion);
 			Device->PendingCommandLists.Reset();
 		}
 		else
 		{
+			// If there are no pending command lists then we assyme that the current command list is open
+			check(bIsCommandListOpen);
+
 			// Just submit the current command list
 			CommandListHandle.Execute(WaitForCompletion);
 		}
 
-		// Get a new command list to replace the one we submitted for execution. 
-		// Restore the state from the previous command list.
-		OpenCommandList();
+		if (bIsCommandListOpen)
+		{
+			// Get a new command list to replace the one we submitted for execution. 
+			// Restore the state from the previous command list.
+			OpenCommandList();
+		}
 	}
 
 	return CommandListHandle;
@@ -393,7 +446,7 @@ void FD3D12CommandContextBase::RHIBeginFrame()
 		bTrackingEvents |= bIsDefaultContext && Device->GetGPUProfiler().bTrackingEvents;
 
 #if D3D12_SUBMISSION_GAP_RECORDER
-		if (GEnableGapRecorder)
+		if (GEnableGapRecorder && !GTriggerGPUProfile)
 		{
 			Device->GetDefaultCommandContext().FlushCommands(bWaitForCommands, FlushAction);
 			GGapRecorderActiveOnBeginFrame = true;
@@ -841,21 +894,36 @@ IRHICommandContextContainer* FD3D12DynamicRHI::RHIGetCommandContextContainer(int
 
 #endif // D3D12_SUPPORTS_PARALLEL_RHI_EXECUTE
 
-void FD3D12DynamicRHI::RHICreateResourceTransition(FRHITransition* Transition, EResourceTransitionPipeline Pipeline, EResourceTransitionPipelineFlags CreateFlags, TArrayView<const FRHITransitionInfo> Infos)
+void FD3D12DynamicRHI::RHICreateTransition(FRHITransition* Transition, ERHIPipeline SrcPipelines, ERHIPipeline DstPipelines, ERHICreateTransitionFlags CreateFlags, TArrayView<const FRHITransitionInfo> Infos)
 {
-	ERHIPipeline SrcPipelines = ERHIPipeline::Graphics;
-	ERHIPipeline DstPipelines = ERHIPipeline::Graphics;
-	switch (Pipeline)
+	checkf(FMath::IsPowerOfTwo(uint32(SrcPipelines)) && FMath::IsPowerOfTwo(uint32(DstPipelines)), TEXT("Support for multi-pipe resources is not yet implemented."));
+
+	// Construct the data in-place on the transition instance
+	FD3D12TransitionData* Data = new (Transition->GetPrivateData<FD3D12TransitionData>()) FD3D12TransitionData;
+
+	Data->SrcPipelines = SrcPipelines;
+	Data->DstPipelines = DstPipelines;
+	Data->CreateFlags = CreateFlags;
+
+	const bool bCrossPipeline = SrcPipelines != DstPipelines;
+
+	if (bCrossPipeline && !EnumHasAnyFlags(Data->CreateFlags, ERHICreateTransitionFlags::NoFence))
 	{
-	case EResourceTransitionPipeline::Graphics_To_Graphics:
-		SrcPipelines = ERHIPipeline::Graphics;
-		DstPipelines = ERHIPipeline::Graphics;
-		break;
+		const FName Name = SrcPipelines == ERHIPipeline::Graphics ? L"<Graphics To AsyncCompute>" : L"<AsyncCompute To Graphics>";
 
-	case EResourceTransitionPipeline::Graphics_To_AsyncCompute:
-		SrcPipelines = ERHIPipeline::Graphics;
-		DstPipelines = ERHIPipeline::AsyncCompute;
-		break;
+		Data->Fence = new FD3D12Fence(&GetAdapter(), FRHIGPUMask::All(), Name);
+		Data->Fence->CreateFence();
+	}
+
+	Data->bCrossPipeline = bCrossPipeline;
+	Data->Infos.Append(Infos.GetData(), Infos.Num());
+}
+
+void FD3D12DynamicRHI::RHIReleaseTransition(FRHITransition* Transition)
+{
+	// Destruct the transition data
+	Transition->GetPrivateData<FD3D12TransitionData>()->~FD3D12TransitionData();
+}
 
 	case EResourceTransitionPipeline::AsyncCompute_To_Graphics:
 		SrcPipelines = ERHIPipeline::AsyncCompute;
@@ -910,20 +978,15 @@ FD3D12CommandContextRedirector::FD3D12CommandContextRedirector(class FD3D12Adapt
 	FMemory::Memzero(PhysicalContexts, sizeof(PhysicalContexts[0]) * MAX_NUM_GPUS);
 }
 
-void FD3D12CommandContextRedirector::RHITransitionResources(EResourceTransitionAccess TransitionType, EResourceTransitionPipeline TransitionPipeline, FRHIUnorderedAccessView** InUAVs, int32 NumUAVs, FRHIComputeFence* WriteComputeFenceRHI)
+void FD3D12CommandContextRedirector::RHIBeginTransitions(TArrayView<const FRHITransition*> Transitions)
 {
-	ContextRedirect(RHITransitionResources(TransitionType, TransitionPipeline, InUAVs, NumUAVs, nullptr));
-
-	// The fence must only be written after evey GPU has transitionned the resource as it handles all GPU.
-	if (WriteComputeFenceRHI)
-	{
-		RHISubmitCommandsHint();
-
-		FD3D12Fence* Fence = FD3D12DynamicRHI::ResourceCast(WriteComputeFenceRHI);
-		Fence->WriteFence();
+	ContextRedirect(RHIBeginTransitionsWithoutFencing(Transitions));
+	SignalTransitionFences(Transitions);
+}
 
-		Fence->Signal(ED3D12CommandQueueType::Default);
-	}	
+void FD3D12CommandContextRedirector::RHIEndTransitions(TArrayView<const FRHITransition*> Transitions)
+{
+	ContextRedirect(RHIEndTransitions(Transitions));
 }
 
 void FD3D12CommandContextRedirector::RHIBeginResourceTransitions(TArrayView<const FRHITransition*> Transitions)
diff --git a/Engine/Source/Runtime/D3D12RHI/Private/D3D12CommandContext.h b/Engine/Source/Runtime/D3D12RHI/Private/D3D12CommandContext.h
--- a/Engine/Source/Runtime/D3D12RHI/Private/D3D12CommandContext.h
+++ b/Engine/Source/Runtime/D3D12RHI/Private/D3D12CommandContext.h
@@ -43,8 +43,6 @@ public:
 
 	virtual void SetRenderTargets(uint32 NumSimultaneousRenderTargets, const FRHIRenderTargetView* NewRenderTargets, const FRHIDepthRenderTargetView* NewDepthStencilTarget) = 0;
 
-	void RHIWaitComputeFence(FRHIComputeFence* InFence) final override;
-
 	virtual void UpdateMemoryStats();
 
 	FRHIGPUMask GetGPUMask() const { return GPUMask; }
@@ -54,7 +52,6 @@ public:
 	virtual void RHISetAsyncComputeBudget(EAsyncComputeBudget Budget) {}
 
 protected:
-
 	virtual FD3D12CommandContext* GetContext(uint32 InGPUIndex) = 0;
 
 	void SignalTransitionFences(TArrayView<const FRHITransition*> Transitions);
@@ -140,7 +137,6 @@ public:
 	FD3D12TextureBase* CurrentDepthTexture;
 	uint32 NumSimultaneousRenderTargets;
 
-
 	/** Track the currently bound uniform buffers. */
 	FD3D12UniformBuffer* BoundUniformBuffers[SF_NumStandardFrequencies][MAX_CBS];
 	FUniformBufferRHIRef BoundUniformBufferRefs[SF_NumStandardFrequencies][MAX_CBS];
@@ -173,6 +169,7 @@ public:
 	virtual void FlushMetadata(FRHITexture** InTextures, int32 NumTextures) {};
 
 	D3D12_RESOURCE_STATES SkipFastClearEliminateState;
+	D3D12_RESOURCE_STATES ValidResourceStates;
 
 #if PLATFORM_SUPPORTS_VIRTUAL_TEXTURES
 	bool bNeedFlushTextureCache;
@@ -244,15 +241,14 @@ public:
 	virtual void SetAsyncComputeBudgetInternal(EAsyncComputeBudget Budget) {}
 
 	void RHIBeginTransitionsWithoutFencing(TArrayView<const FRHITransition*> Transitions);
-	virtual void RHIBeginResourceTransitions(TArrayView<const FRHITransition*> Transitions) final override;
-	virtual void RHIEndResourceTransitions(TArrayView<const FRHITransition*> Transitions) final override;
+	virtual void RHIBeginTransitions(TArrayView<const FRHITransition*> Transitions) final override;
+	virtual void RHIEndTransitions(TArrayView<const FRHITransition*> Transitions) final override;
 
 	// IRHIComputeContext interface
 	virtual void RHISetComputeShader(FRHIComputeShader* ComputeShader) final override;
 	virtual void RHISetComputePipelineState(FRHIComputePipelineState* ComputePipelineState) final override;
 	virtual void RHIDispatchComputeShader(uint32 ThreadGroupCountX, uint32 ThreadGroupCountY, uint32 ThreadGroupCountZ) final override;
 	virtual void RHIDispatchIndirectComputeShader(FRHIVertexBuffer* ArgumentBuffer, uint32 ArgumentOffset) final override;
-	virtual void RHITransitionResources(EResourceTransitionAccess TransitionType, EResourceTransitionPipeline TransitionPipeline, FRHIUnorderedAccessView** InUAVs, int32 NumUAVs, FRHIComputeFence* WriteComputeFenceRHI) final override;
 	virtual void RHISetGlobalUniformBuffers(const FUniformBufferStaticBindings& InUniformBuffers) final override;
 	virtual void RHISetShaderTexture(FRHIComputeShader* PixelShader, uint32 TextureIndex, FRHITexture* NewTexture) final override;
 	virtual void RHISetShaderSampler(FRHIComputeShader* ComputeShader, uint32 SamplerIndex, FRHISamplerState* NewState) final override;
@@ -267,14 +263,11 @@ public:
 	virtual void RHISubmitCommandsHint() final override;
 
 	// IRHICommandContext interface
-	virtual void RHIAutomaticCacheFlushAfterComputeShader(bool bEnable) final override;
-	virtual void RHIFlushComputeShaderCache() final override;
 	virtual void RHISetMultipleViewports(uint32 Count, const FViewportBounds* Data) final override;
 	virtual void RHIClearUAVFloat(FRHIUnorderedAccessView* UnorderedAccessViewRHI, const FVector4& Values) final override;
 	virtual void RHIClearUAVUint(FRHIUnorderedAccessView* UnorderedAccessViewRHI, const FUintVector4& Values) final override;
 	virtual void RHICopyToResolveTarget(FRHITexture* SourceTexture, FRHITexture* DestTexture, const FResolveParams& ResolveParams) override;
 	virtual void RHICopyTexture(FRHITexture* SourceTexture, FRHITexture* DestTexture, const FRHICopyTextureInfo& CopyInfo) final override;
-	virtual void RHITransitionResources(EResourceTransitionAccess TransitionType, FRHITexture** InTextures, int32 NumTextures) final override;
 	virtual void RHICopyToStagingBuffer(FRHIVertexBuffer* SourceBuffer, FRHIStagingBuffer* DestinationStagingBuffer, uint32 Offset, uint32 NumBytes) final override;
 	virtual void RHIWriteGPUFence(FRHIGPUFence* Fence) final override;
 	virtual void RHIBeginRenderQuery(FRHIRenderQuery* RenderQuery) final override;
@@ -297,7 +290,6 @@ public:
 	virtual void RHISetBlendFactor(const FLinearColor& BlendFactor) final override;
 	virtual void SetRenderTargets(uint32 NumSimultaneousRenderTargets, const FRHIRenderTargetView* NewRenderTargets, const FRHIDepthRenderTargetView* NewDepthStencilTarget) final override;
 	void SetRenderTargetsAndClear(const FRHISetRenderTargetsInfo& RenderTargetsInfo);
-	virtual void RHIBindClearMRTValues(bool bClearColor, bool bClearDepth, bool bClearStencil) final override;
 	virtual void RHIDrawPrimitive(uint32 BaseVertexIndex, uint32 NumPrimitives, uint32 NumInstances) final override;
 	virtual void RHIDrawPrimitiveIndirect(FRHIVertexBuffer* ArgumentBuffer, uint32 ArgumentOffset) final override;
 	virtual void RHIDrawIndexedIndirect(FRHIIndexBuffer* IndexBufferRHI, FRHIStructuredBuffer* ArgumentsBufferRHI, int32 DrawArgumentsIndex, uint32 NumInstances) final override;
@@ -316,29 +308,6 @@ public:
 
 	virtual void RHIBeginRenderPass(const FRHIRenderPassInfo& InInfo, const TCHAR* InName)
 	{
-		if (InInfo.bGeneratingMips)
-		{
-			FRHITexture* Textures[MaxSimultaneousRenderTargets];
-			FRHITexture** LastTexture = Textures;
-			for (int32 Index = 0; Index < MaxSimultaneousRenderTargets; ++Index)
-			{
-				if (!InInfo.ColorRenderTargets[Index].RenderTarget)
-				{
-					break;
-				}
-
-				*LastTexture = InInfo.ColorRenderTargets[Index].RenderTarget;
-				++LastTexture;
-			}
-
-			//Use RWBarrier since we don't transition individual subresources.  Basically treat the whole texture as R/W as we walk down the mip chain.
-			int32 NumTextures = (int32)(LastTexture - Textures);
-			if (NumTextures)
-			{
-				RHITransitionResources(EResourceTransitionAccess::ERWSubResBarrier, Textures, NumTextures);
-			}
-		}
-
 		FRHISetRenderTargetsInfo RTInfo;
 		InInfo.ConvertToRenderTargetsInfo(RTInfo);
 		SetRenderTargetsAndClear(RTInfo);
@@ -442,7 +411,8 @@ public:
 
 	static inline FD3D12TextureBase* RetrieveTextureBase(FRHITexture* Texture, uint32 GPUIndex)
 	{
-		return Texture ? static_cast<FD3D12TextureBase*>(Texture->GetTextureBaseRHI())->GetLinkedObject(GPUIndex) : nullptr;
+		FD3D12TextureBase* RHITexture = GetD3D12TextureFromRHITexture(Texture);
+		return RHITexture ? RHITexture->GetLinkedObject(GPUIndex) : nullptr;
 	}
 
 	FORCEINLINE_DEBUGGABLE FD3D12TextureBase* RetrieveTextureBase(FRHITexture* Texture)
@@ -514,8 +484,10 @@ public:
 	{
 		ContextRedirect(RHIDispatchIndirectComputeShader(ArgumentBuffer, ArgumentOffset));
 	}
+
 	// Special implementation that only signal the fence once.
-	virtual void RHITransitionResources(EResourceTransitionAccess TransitionType, EResourceTransitionPipeline TransitionPipeline, FRHIUnorderedAccessView** InUAVs, int32 NumUAVs, FRHIComputeFence* WriteComputeFenceRHI) final override;
+	virtual void RHIBeginTransitions(TArrayView<const FRHITransition*> Transitions) final override;
+	virtual void RHIEndTransitions(TArrayView<const FRHITransition*> Transitions) final override;
 
 	// Special implementation that only signal the fence once.
 	virtual void RHIBeginResourceTransitions(TArrayView<const FRHITransition*> Transitions) final override;
@@ -575,14 +547,6 @@ public:
 	}
 
 	// IRHICommandContext interface
-	FORCEINLINE virtual void RHIAutomaticCacheFlushAfterComputeShader(bool bEnable) final override
-	{
-		ContextRedirect(RHIAutomaticCacheFlushAfterComputeShader(bEnable));
-	}
-	FORCEINLINE virtual void RHIFlushComputeShaderCache() final override
-	{
-		ContextRedirect(RHIFlushComputeShaderCache());
-	}
 	FORCEINLINE virtual void RHISetMultipleViewports(uint32 Count, const FViewportBounds* Data) final override
 	{
 		ContextRedirect(RHISetMultipleViewports(Count, Data));
@@ -603,10 +567,6 @@ public:
 	{
 		ContextRedirect(RHICopyTexture(SourceTextureRHI, DestTextureRHI, CopyInfo));
 	}
-	FORCEINLINE virtual void RHITransitionResources(EResourceTransitionAccess TransitionType, FRHITexture** InTextures, int32 NumTextures) final override
-	{
-		ContextRedirect(RHITransitionResources(TransitionType, InTextures, NumTextures));
-	}
 	FORCEINLINE virtual void RHIBeginRenderQuery(FRHIRenderQuery* RenderQuery) final override
 	{
 		ContextRedirect(RHIBeginRenderQuery(RenderQuery));
@@ -683,10 +643,6 @@ public:
 	{
 		ContextRedirect(SetRenderTargetsAndClear(RenderTargetsInfo));
 	}
-	FORCEINLINE virtual void RHIBindClearMRTValues(bool bClearColor, bool bClearDepth, bool bClearStencil) final override
-	{
-		ContextRedirect(RHIBindClearMRTValues(bClearColor, bClearDepth, bClearStencil));
-	}
 	FORCEINLINE virtual void RHIDrawPrimitive(uint32 BaseVertexIndex, uint32 NumPrimitives, uint32 NumInstances) final override
 	{
 		ContextRedirect(RHIDrawPrimitive(BaseVertexIndex, NumPrimitives, NumInstances));
@@ -906,7 +862,7 @@ private:
 struct FD3D12TransitionData
 {
 	ERHIPipeline SrcPipelines, DstPipelines;
-	EResourceTransitionPipelineFlags CreateFlags = EResourceTransitionPipelineFlags::None;
+	ERHICreateTransitionFlags CreateFlags = ERHICreateTransitionFlags::None;
 
 	TArray<FRHITransitionInfo, TInlineAllocator<4>> Infos;
 	TRefCountPtr<FD3D12Fence> Fence;
diff --git a/Engine/Source/Runtime/D3D12RHI/Private/D3D12RHIPrivate.h b/Engine/Source/Runtime/D3D12RHI/Private/D3D12RHIPrivate.h
--- a/Engine/Source/Runtime/D3D12RHI/Private/D3D12RHIPrivate.h
+++ b/Engine/Source/Runtime/D3D12RHI/Private/D3D12RHIPrivate.h
@@ -47,11 +47,18 @@ DECLARE_LOG_CATEGORY_EXTERN(LogD3D12GapRecorder, Log, All);
 #include "D3D12RHIBasePrivate.h"
 #endif
 
+#if !defined(NV_AFTERMATH)
+	#define NV_AFTERMATH 0
+#endif
+
 #if NV_AFTERMATH
 #define GFSDK_Aftermath_WITH_DX12 1
 #include "GFSDK_Aftermath.h"
+#include "GFSDK_Aftermath_GpuCrashdump.h"
 #undef GFSDK_Aftermath_WITH_DX12
 extern int32 GDX12NVAfterMathEnabled;
+extern int32 GDX12NVAfterMathTrackResources;
+extern int32 GDX12NVAfterMathMarkers;
 #endif
 
 #include "D3D12Residency.h"
@@ -102,6 +109,9 @@ typedef FD3D12StateCacheBase FD3D12StateCache;
 #endif //D3D12_RHI_RAYTRACING
 #endif //DEFAULT_BUFFER_POOL_MAX_ALLOC_SIZE
 
+#define READBACK_BUFFER_POOL_MAX_ALLOC_SIZE (64 * 1024)
+#define READBACK_BUFFER_POOL_DEFAULT_POOL_SIZE (4 * 1024 * 1024)
+
 #define DEFAULT_CONTEXT_UPLOAD_POOL_SIZE (8 * 1024 * 1024)
 #define DEFAULT_CONTEXT_UPLOAD_POOL_MAX_ALLOC_SIZE (4 * 1024 * 1024)
 #define DEFAULT_CONTEXT_UPLOAD_POOL_ALIGNMENT (D3D12_CONSTANT_BUFFER_DATA_PLACEMENT_ALIGNMENT)
@@ -199,7 +209,7 @@ class FD3D12DynamicRHI : public FDynamicRHI
 
 public:
 
-	static FD3D12DynamicRHI* GetD3DRHI() { return SingleD3DRHI; }
+	static D3D12RHI_API FD3D12DynamicRHI* GetD3DRHI() { return SingleD3DRHI; }
 
 private:
 
@@ -254,7 +264,6 @@ public:
 	virtual FDomainShaderRHIRef RHICreateDomainShader(TArrayView<const uint8> Code, const FSHAHash& Hash) final override;
 	virtual FGeometryShaderRHIRef RHICreateGeometryShader(TArrayView<const uint8> Code, const FSHAHash& Hash) final override;
 	virtual FComputeShaderRHIRef RHICreateComputeShader(TArrayView<const uint8> Code, const FSHAHash& Hash) override;
-	virtual FComputeFenceRHIRef RHICreateComputeFence(const FName& Name) final override;
 	virtual FGPUFenceRHIRef RHICreateGPUFence(const FName& Name) final override;
 	virtual FStagingBufferRHIRef RHICreateStagingBuffer() final override;
 	virtual void* RHILockStagingBuffer(FRHIStagingBuffer* StagingBuffer, FRHIGPUFence* Fence, uint32 Offset, uint32 SizeRHI) final override;
@@ -262,20 +271,20 @@ public:
 	virtual FBoundShaderStateRHIRef RHICreateBoundShaderState(FRHIVertexDeclaration* VertexDeclaration, FRHIVertexShader* VertexShader, FRHIHullShader* HullShader, FRHIDomainShader* DomainShader, FRHIPixelShader* PixelShader, FRHIGeometryShader* GeometryShader) final override;
 	virtual FGraphicsPipelineStateRHIRef RHICreateGraphicsPipelineState(const FGraphicsPipelineStateInitializer& Initializer) final override;
 	virtual TRefCountPtr<FRHIComputePipelineState> RHICreateComputePipelineState(FRHIComputeShader* ComputeShader) final override;
-	virtual void RHICreateResourceTransition(FRHITransition* Transition, EResourceTransitionPipeline Pipeline, EResourceTransitionPipelineFlags CreateFlags, TArrayView<const FRHITransitionInfo> Infos) final override;
-	virtual void RHIReleaseResourceTransition(FRHITransition* Transition) final override;
+	virtual void RHICreateTransition(FRHITransition* Transition, ERHIPipeline SrcPipelines, ERHIPipeline DstPipelines, ERHICreateTransitionFlags CreateFlags, TArrayView<const FRHITransitionInfo> Infos) final override;
+	virtual void RHIReleaseTransition(FRHITransition* Transition) final override;
 	virtual FUniformBufferRHIRef RHICreateUniformBuffer(const void* Contents, const FRHIUniformBufferLayout& Layout, EUniformBufferUsage Usage, EUniformBufferValidation Validation) final override;
 	virtual void RHIUpdateUniformBuffer(FRHIUniformBuffer* UniformBufferRHI, const void* Contents) final override;
-	virtual FIndexBufferRHIRef RHICreateIndexBuffer(uint32 Stride, uint32 Size, uint32 InUsage, FRHIResourceCreateInfo& CreateInfo) final override;
+	virtual FIndexBufferRHIRef RHICreateIndexBuffer(uint32 Stride, uint32 Size, uint32 InUsage, ERHIAccess InResourceState, FRHIResourceCreateInfo& CreateInfo) final override;
 	virtual void* RHILockIndexBuffer(FRHICommandListImmediate& RHICmdList, FRHIIndexBuffer* IndexBuffer, uint32 Offset, uint32 Size, EResourceLockMode LockMode) final override;
 	virtual void RHIUnlockIndexBuffer(FRHICommandListImmediate& RHICmdList, FRHIIndexBuffer* IndexBuffer) final override;
 	virtual void RHITransferIndexBufferUnderlyingResource(FRHIIndexBuffer* DestIndexBuffer, FRHIIndexBuffer* SrcIndexBuffer) final override;
-	virtual FVertexBufferRHIRef RHICreateVertexBuffer(uint32 Size, uint32 InUsage, FRHIResourceCreateInfo& CreateInfo) final override;
+	virtual FVertexBufferRHIRef RHICreateVertexBuffer(uint32 Size, uint32 InUsage, ERHIAccess InResourceState, FRHIResourceCreateInfo& CreateInfo) final override;
 	virtual void* RHILockVertexBuffer(FRHICommandListImmediate& RHICmdList, FRHIVertexBuffer* VertexBuffer, uint32 Offset, uint32 SizeRHI, EResourceLockMode LockMode) final override;
 	virtual void RHIUnlockVertexBuffer(FRHICommandListImmediate& RHICmdList, FRHIVertexBuffer* VertexBuffer) final override;
 	virtual void RHICopyVertexBuffer(FRHIVertexBuffer* SourceBuffer, FRHIVertexBuffer* DestBuffer) final override;
 	virtual void RHITransferVertexBufferUnderlyingResource(FRHIVertexBuffer* DestVertexBuffer, FRHIVertexBuffer* SrcVertexBuffer) final override;
-	virtual FStructuredBufferRHIRef RHICreateStructuredBuffer(uint32 Stride, uint32 Size, uint32 InUsage, FRHIResourceCreateInfo& CreateInfo) final override;
+	virtual FStructuredBufferRHIRef RHICreateStructuredBuffer(uint32 Stride, uint32 Size, uint32 InUsage, ERHIAccess InResourceState, FRHIResourceCreateInfo& CreateInfo) final override;
 	virtual void* RHILockStructuredBuffer(FRHICommandListImmediate& RHICmdList, FRHIStructuredBuffer* StructuredBuffer, uint32 Offset, uint32 SizeRHI, EResourceLockMode LockMode) final override;
 	virtual void RHIUnlockStructuredBuffer(FRHICommandListImmediate& RHICmdList, FRHIStructuredBuffer* StructuredBuffer) final override;
 	virtual FUnorderedAccessViewRHIRef RHICreateUnorderedAccessView(FRHIStructuredBuffer* StructuredBuffer, bool bUseUAVCounter, bool bAppendBuffer) final override;
@@ -288,21 +297,21 @@ public:
 	virtual FShaderResourceViewRHIRef RHICreateShaderResourceView(FRHIIndexBuffer* Buffer) final override;
 	virtual void RHIUpdateShaderResourceView(FRHIShaderResourceView* SRV, FRHIVertexBuffer* VertexBuffer, uint32 Stride, uint8 Format) final override;
 	virtual void RHIUpdateShaderResourceView(FRHIShaderResourceView* SRV, FRHIIndexBuffer* IndexBuffer) final override;
-	virtual uint64 RHICalcTexture2DPlatformSize(uint32 SizeX, uint32 SizeY, uint8 Format, uint32 NumMips, uint32 NumSamples, uint32 Flags, const FRHIResourceCreateInfo& CreateInfo, uint32& OutAlign) override;
-	virtual uint64 RHICalcTexture3DPlatformSize(uint32 SizeX, uint32 SizeY, uint32 SizeZ, uint8 Format, uint32 NumMips, uint32 Flags, const FRHIResourceCreateInfo& CreateInfo, uint32& OutAlign) final override;
-	virtual uint64 RHICalcTextureCubePlatformSize(uint32 Size, uint8 Format, uint32 NumMips, uint32 Flags, const FRHIResourceCreateInfo& CreateInfo, uint32& OutAlign) final override;
+	virtual uint64 RHICalcTexture2DPlatformSize(uint32 SizeX, uint32 SizeY, uint8 Format, uint32 NumMips, uint32 NumSamples, ETextureCreateFlags InFlags, const FRHIResourceCreateInfo& CreateInfo, uint32& OutAlign) override;
+	virtual uint64 RHICalcTexture3DPlatformSize(uint32 SizeX, uint32 SizeY, uint32 SizeZ, uint8 Format, uint32 NumMips, ETextureCreateFlags InFlags, const FRHIResourceCreateInfo& CreateInfo, uint32& OutAlign) final override;
+	virtual uint64 RHICalcTextureCubePlatformSize(uint32 Size, uint8 Format, uint32 NumMips, ETextureCreateFlags InFlags, const FRHIResourceCreateInfo& CreateInfo, uint32& OutAlign) final override;
 	virtual uint64 RHIGetMinimumAlignmentForBufferBackedSRV(EPixelFormat Format) final override;
 	virtual void RHIGetTextureMemoryStats(FTextureMemoryStats& OutStats) final override;
 	virtual bool RHIGetTextureMemoryVisualizeData(FColor* TextureData, int32 SizeX, int32 SizeY, int32 Pitch, int32 PixelSize) final override;
 	virtual FTextureReferenceRHIRef RHICreateTextureReference(FLastRenderTimeContainer* LastRenderTime) final override;
-	virtual FTexture2DRHIRef RHICreateTexture2D(uint32 SizeX, uint32 SizeY, uint8 Format, uint32 NumMips, uint32 NumSamples, uint32 Flags, FRHIResourceCreateInfo& CreateInfo) override;
-	virtual FTexture2DRHIRef RHIAsyncCreateTexture2D(uint32 SizeX, uint32 SizeY, uint8 Format, uint32 NumMips, uint32 Flags, void** InitialMipData, uint32 NumInitialMips) final override;
+	virtual FTexture2DRHIRef RHICreateTexture2D(uint32 SizeX, uint32 SizeY, uint8 Format, uint32 NumMips, uint32 NumSamples, ETextureCreateFlags Flags, ERHIAccess InResourceState, FRHIResourceCreateInfo& CreateInfo) override;
+	virtual FTexture2DRHIRef RHIAsyncCreateTexture2D(uint32 SizeX, uint32 SizeY, uint8 Format, uint32 NumMips, ETextureCreateFlags Flags, ERHIAccess InResourceState, void** InitialMipData, uint32 NumInitialMips) final override;
 	virtual void RHICopySharedMips(FRHITexture2D* DestTexture2D, FRHITexture2D* SrcTexture2D) final override;
 	void RHIMultiGPULockstep(FRHIGPUMask GPUMask);
 	virtual void RHITransferTexture(FRHITexture2D* Texture, FIntRect Rect, uint32 SrcGPUIndex, uint32 DestGPUIndex, bool PullData) final override;
 	virtual void RHITransferTextures(const TArrayView<const FTransferTextureParams> Params) final override;
-	virtual FTexture2DArrayRHIRef RHICreateTexture2DArray(uint32 SizeX, uint32 SizeY, uint32 SizeZ, uint8 Format, uint32 NumMips, uint32 NumSamples, uint32 Flags, FRHIResourceCreateInfo& CreateInfo) override;
-	virtual FTexture3DRHIRef RHICreateTexture3D(uint32 SizeX, uint32 SizeY, uint32 SizeZ, uint8 Format, uint32 NumMips, uint32 Flags, FRHIResourceCreateInfo& CreateInfo) override;
+	virtual FTexture2DArrayRHIRef RHICreateTexture2DArray(uint32 SizeX, uint32 SizeY, uint32 SizeZ, uint8 Format, uint32 NumMips, uint32 NumSamples, ETextureCreateFlags Flags, ERHIAccess InResourceState, FRHIResourceCreateInfo& CreateInfo) override;
+	virtual FTexture3DRHIRef RHICreateTexture3D(uint32 SizeX, uint32 SizeY, uint32 SizeZ, uint8 Format, uint32 NumMips, ETextureCreateFlags Flags, ERHIAccess InResourceState, FRHIResourceCreateInfo& CreateInfo) override;
 	virtual void RHIGetResourceInfo(FRHITexture* Ref, FRHIResourceInfo& OutInfo) override;
 	virtual FShaderResourceViewRHIRef RHICreateShaderResourceView(FRHITexture* Texture, const FRHITextureSRVCreateInfo& CreateInfo) final override;
 	virtual uint32 RHIComputeMemorySize(FRHITexture* TextureRHI) final override;
@@ -318,8 +327,8 @@ public:
 	virtual FUpdateTexture3DData BeginUpdateTexture3D_RenderThread(class FRHICommandListImmediate& RHICmdList, FRHITexture3D* Texture, uint32 MipIndex, const struct FUpdateTextureRegion3D& UpdateRegion) final override;
 	virtual void EndUpdateTexture3D_RenderThread(class FRHICommandListImmediate& RHICmdList, FUpdateTexture3DData& UpdateData) final override;
 	virtual void EndMultiUpdateTexture3D_RenderThread(class FRHICommandListImmediate& RHICmdList, TArray<FUpdateTexture3DData>& UpdateDataArray) final override;
-	virtual FTextureCubeRHIRef RHICreateTextureCube(uint32 Size, uint8 Format, uint32 NumMips, uint32 Flags, FRHIResourceCreateInfo& CreateInfo) override;
-	virtual FTextureCubeRHIRef RHICreateTextureCubeArray(uint32 Size, uint32 ArraySize, uint8 Format, uint32 NumMips, uint32 Flags, FRHIResourceCreateInfo& CreateInfo) override;
+	virtual FTextureCubeRHIRef RHICreateTextureCube(uint32 Size, uint8 Format, uint32 NumMips, ETextureCreateFlags Flags, ERHIAccess InResourceState, FRHIResourceCreateInfo& CreateInfo) override;
+	virtual FTextureCubeRHIRef RHICreateTextureCubeArray(uint32 Size, uint32 ArraySize, uint8 Format, uint32 NumMips, ETextureCreateFlags Flags, ERHIAccess InResourceState, FRHIResourceCreateInfo& CreateInfo) override;
 	virtual void* RHILockTextureCubeFace(FRHITextureCube* Texture, uint32 FaceIndex, uint32 ArrayIndex, uint32 MipIndex, EResourceLockMode LockMode, uint32& DestStride, bool bLockWithinMiptail) final override;
 	virtual void RHIUnlockTextureCubeFace(FRHITextureCube* Texture, uint32 FaceIndex, uint32 ArrayIndex, uint32 MipIndex, bool bLockWithinMiptail) final override;
 	virtual void RHIBindDebugLabelName(FRHITexture* Texture, const TCHAR* Name) final override;
@@ -327,8 +336,10 @@ public:
 	virtual void RHIReadSurfaceData(FRHITexture* TextureRHI, FIntRect InRect, TArray<FLinearColor>& OutData, FReadSurfaceDataFlags InFlags) final override;
 	virtual void RHIMapStagingSurface(FRHITexture* Texture, FRHIGPUFence* Fence, void*& OutData, int32& OutWidth, int32& OutHeight, uint32 GPUIndex = 0) final override;
 	virtual void RHIUnmapStagingSurface(FRHITexture* Texture, uint32 GPUIndex = 0) final override;
+	virtual void RHIReadSurfaceFloatData(FRHITexture* Texture, FIntRect Rect, TArray<FFloat16Color>& OutData, FReadSurfaceDataFlags InFlags) final override;
 	virtual void RHIReadSurfaceFloatData(FRHITexture* Texture, FIntRect Rect, TArray<FFloat16Color>& OutData, ECubeFace CubeFace, int32 ArrayIndex, int32 MipIndex) final override;
 	virtual void RHIRead3DSurfaceFloatData(FRHITexture* Texture, FIntRect Rect, FIntPoint ZMinMax, TArray<FFloat16Color>& OutData) final override;
+	virtual void RHIRead3DSurfaceFloatData(FRHITexture* Texture, FIntRect Rect, FIntPoint ZMinMax, TArray<FFloat16Color>& OutData, FReadSurfaceDataFlags InFlags) final override;
 	virtual FRenderQueryRHIRef RHICreateRenderQuery(ERenderQueryType QueryType) final override;
 	virtual bool RHIGetRenderQueryResult(FRHIRenderQuery* RenderQuery, uint64& OutResult, bool bWait, uint32 GPUIndex = INDEX_NONE) final override;
 	virtual uint32 RHIGetViewportNextPresentGPUIndex(FRHIViewport* Viewport) final override;
@@ -369,8 +380,9 @@ public:
 	// FD3D12DynamicRHI interface.
 	virtual uint32 GetDebugFlags();
 	virtual ID3D12CommandQueue* RHIGetD3DCommandQueue();
-	virtual FTexture2DRHIRef RHICreateTexture2DFromResource(EPixelFormat Format, uint32 TexCreateFlags, const FClearValueBinding& ClearValueBinding, ID3D12Resource* Resource);
-	virtual FTextureCubeRHIRef RHICreateTextureCubeFromResource(EPixelFormat Format, uint32 TexCreateFlags, const FClearValueBinding& ClearValueBinding, ID3D12Resource* Resource);
+	virtual FTexture2DRHIRef RHICreateTexture2DFromResource(EPixelFormat Format, ETextureCreateFlags TexCreateFlags, const FClearValueBinding& ClearValueBinding, ID3D12Resource* Resource);
+	virtual FTexture2DRHIRef RHICreateTexture2DArrayFromResource(EPixelFormat Format, ETextureCreateFlags TexCreateFlags, const FClearValueBinding& ClearValueBinding, ID3D12Resource* Resource);
+	virtual FTextureCubeRHIRef RHICreateTextureCubeFromResource(EPixelFormat Format, ETextureCreateFlags TexCreateFlags, const FClearValueBinding& ClearValueBinding, ID3D12Resource* Resource);
 
 	//
 	// The Following functions are the _RenderThread version of the above functions. They allow the RHI to control the thread synchronization for greater efficiency.
@@ -407,15 +419,18 @@ public:
 		return RHICreateComputeShader(Code, Hash);
 	}
 
-	virtual FVertexBufferRHIRef CreateVertexBuffer_RenderThread(class FRHICommandListImmediate& RHICmdList, uint32 Size, uint32 InUsage, FRHIResourceCreateInfo& CreateInfo);
-	virtual FStructuredBufferRHIRef CreateStructuredBuffer_RenderThread(class FRHICommandListImmediate& RHICmdList, uint32 Stride, uint32 Size, uint32 InUsage, FRHIResourceCreateInfo& CreateInfo);
-	virtual FVertexBufferRHIRef CreateAndLockVertexBuffer_RenderThread(class FRHICommandListImmediate& RHICmdList, uint32 Size, uint32 InUsage, FRHIResourceCreateInfo& CreateInfo, void*& OutDataBuffer);
-	virtual FIndexBufferRHIRef CreateAndLockIndexBuffer_RenderThread(class FRHICommandListImmediate& RHICmdList, uint32 Stride, uint32 Size, uint32 InUsage, FRHIResourceCreateInfo& CreateInfo, void*& OutDataBuffer);	
-	virtual FIndexBufferRHIRef CreateIndexBuffer_RenderThread(class FRHICommandListImmediate& RHICmdList, uint32 Stride, uint32 Size, uint32 InUsage, FRHIResourceCreateInfo& CreateInfo);
-	virtual void UpdateTexture2D_RenderThread(class FRHICommandListImmediate& RHICmdList, FRHITexture2D* Texture, uint32 MipIndex, const struct FUpdateTextureRegion2D& UpdateRegion, uint32 SourcePitch, const uint8* SourceData);
-	virtual void* LockTexture2D_RenderThread(class FRHICommandListImmediate& RHICmdList, FRHITexture2D* Texture, uint32 MipIndex, EResourceLockMode LockMode, uint32& DestStride, bool bLockWithinMiptail, bool bNeedsDefaultRHIFlush = true);
-	virtual void UnlockTexture2D_RenderThread(class FRHICommandListImmediate& RHICmdList, FRHITexture2D* Texture, uint32 MipIndex, bool bLockWithinMiptail, bool bNeedsDefaultRHIFlush = true);
-	
+	virtual FVertexBufferRHIRef CreateVertexBuffer_RenderThread(class FRHICommandListImmediate& RHICmdList, uint32 Size, uint32 InUsage, ERHIAccess InResourceState, FRHIResourceCreateInfo& CreateInfo) override final;
+	virtual FStructuredBufferRHIRef CreateStructuredBuffer_RenderThread(class FRHICommandListImmediate& RHICmdList, uint32 Stride, uint32 Size, uint32 InUsage, ERHIAccess InResourceState, FRHIResourceCreateInfo& CreateInfo) override final;
+	virtual FVertexBufferRHIRef CreateAndLockVertexBuffer_RenderThread(class FRHICommandListImmediate& RHICmdList, uint32 Size, uint32 InUsage, ERHIAccess InResourceState, FRHIResourceCreateInfo& CreateInfo, void*& OutDataBuffer) override final;
+	virtual FIndexBufferRHIRef CreateAndLockIndexBuffer_RenderThread(class FRHICommandListImmediate& RHICmdList, uint32 Stride, uint32 Size, uint32 InUsage, ERHIAccess InResourceState, FRHIResourceCreateInfo& CreateInfo, void*& OutDataBuffer) override final;
+	// virtual FTexture2DRHIRef AsyncReallocateTexture2D_RenderThread(class FRHICommandListImmediate& RHICmdList, FRHITexture2D* Texture2D, int32 NewMipCount, int32 NewSizeX, int32 NewSizeY, FThreadSafeCounter* RequestStatus) override final;
+	// virtual ETextureReallocationStatus FinalizeAsyncReallocateTexture2D_RenderThread(class FRHICommandListImmediate& RHICmdList, FRHITexture2D* Texture2D, bool bBlockUntilCompleted) override final;
+	// virtual ETextureReallocationStatus CancelAsyncReallocateTexture2D_RenderThread(class FRHICommandListImmediate& RHICmdList, FRHITexture2D* Texture2D, bool bBlockUntilCompleted) override final;
+	virtual FIndexBufferRHIRef CreateIndexBuffer_RenderThread(class FRHICommandListImmediate& RHICmdList, uint32 Stride, uint32 Size, uint32 InUsage, ERHIAccess InResourceState, FRHIResourceCreateInfo& CreateInfo) override final;
+	virtual void UpdateTexture2D_RenderThread(class FRHICommandListImmediate& RHICmdList, FRHITexture2D* Texture, uint32 MipIndex, const struct FUpdateTextureRegion2D& UpdateRegion, uint32 SourcePitch, const uint8* SourceData) override final;
+	virtual void* LockTexture2D_RenderThread(class FRHICommandListImmediate& RHICmdList, FRHITexture2D* Texture, uint32 MipIndex, EResourceLockMode LockMode, uint32& DestStride, bool bLockWithinMiptail, bool bNeedsDefaultRHIFlush = true) override final;
+	virtual void UnlockTexture2D_RenderThread(class FRHICommandListImmediate& RHICmdList, FRHITexture2D* Texture, uint32 MipIndex, bool bLockWithinMiptail, bool bNeedsDefaultRHIFlush = true) override final;
+
 	virtual FTexture2DRHIRef AsyncReallocateTexture2D_RenderThread(class FRHICommandListImmediate& RHICmdList, FRHITexture2D* Texture2D, int32 NewMipCount, int32 NewSizeX, int32 NewSizeY, FThreadSafeCounter* RequestStatus);
 	virtual ETextureReallocationStatus FinalizeAsyncReallocateTexture2D_RenderThread(class FRHICommandListImmediate& RHICmdList, FRHITexture2D* Texture2D, bool bBlockUntilCompleted)
 	{
@@ -426,22 +441,22 @@ public:
 		return RHICancelAsyncReallocateTexture2D(Texture2D, bBlockUntilCompleted);
 	}
 
-	virtual FTexture2DRHIRef RHICreateTexture2D_RenderThread(class FRHICommandListImmediate& RHICmdList, uint32 SizeX, uint32 SizeY, uint8 Format, uint32 NumMips, uint32 NumSamples, uint32 Flags, FRHIResourceCreateInfo& CreateInfo);
-	virtual FTexture2DArrayRHIRef RHICreateTexture2DArray_RenderThread(class FRHICommandListImmediate& RHICmdList, uint32 SizeX, uint32 SizeY, uint32 SizeZ, uint8 Format, uint32 NumMips, uint32 NumSamples, uint32 Flags, FRHIResourceCreateInfo& CreateInfo);
-	virtual FTexture3DRHIRef RHICreateTexture3D_RenderThread(class FRHICommandListImmediate& RHICmdList, uint32 SizeX, uint32 SizeY, uint32 SizeZ, uint8 Format, uint32 NumMips, uint32 Flags, FRHIResourceCreateInfo& CreateInfo);
-	virtual FUnorderedAccessViewRHIRef RHICreateUnorderedAccessView_RenderThread(class FRHICommandListImmediate& RHICmdList, FRHIStructuredBuffer* StructuredBuffer, bool bUseUAVCounter, bool bAppendBuffer);
-	virtual FUnorderedAccessViewRHIRef RHICreateUnorderedAccessView_RenderThread(class FRHICommandListImmediate& RHICmdList, FRHITexture* Texture, uint32 MipLevel);
-	virtual FUnorderedAccessViewRHIRef RHICreateUnorderedAccessView_RenderThread(class FRHICommandListImmediate& RHICmdList, FRHIVertexBuffer* VertexBuffer, uint8 Format);
-	virtual FShaderResourceViewRHIRef RHICreateShaderResourceView_RenderThread(class FRHICommandListImmediate& RHICmdList, FRHITexture* Texture, const FRHITextureSRVCreateInfo& CreateInfo);
-	virtual FShaderResourceViewRHIRef RHICreateShaderResourceView_RenderThread(class FRHICommandListImmediate& RHICmdList, FRHIVertexBuffer* VertexBuffer, uint32 Stride, uint8 Format);
-	virtual FShaderResourceViewRHIRef RHICreateShaderResourceView_RenderThread(class FRHICommandListImmediate& RHICmdList, const FShaderResourceViewInitializer& Initializer);
-	virtual FShaderResourceViewRHIRef RHICreateShaderResourceViewWriteMask_RenderThread(class FRHICommandListImmediate& RHICmdList, FRHITexture2D* Texture2D);
-
-	virtual FShaderResourceViewRHIRef CreateShaderResourceView_RenderThread(class FRHICommandListImmediate& RHICmdList, FRHIVertexBuffer* VertexBuffer, uint32 Stride, uint8 Format);
-	virtual FShaderResourceViewRHIRef CreateShaderResourceView_RenderThread(class FRHICommandListImmediate& RHICmdList, const FShaderResourceViewInitializer& Initializer);
-	virtual FShaderResourceViewRHIRef RHICreateShaderResourceView_RenderThread(class FRHICommandListImmediate& RHICmdList, FRHIStructuredBuffer* StructuredBuffer);
-	virtual FTextureCubeRHIRef RHICreateTextureCube_RenderThread(class FRHICommandListImmediate& RHICmdList, uint32 Size, uint8 Format, uint32 NumMips, uint32 Flags, FRHIResourceCreateInfo& CreateInfo);
-	virtual FTextureCubeRHIRef RHICreateTextureCubeArray_RenderThread(class FRHICommandListImmediate& RHICmdList, uint32 Size, uint32 ArraySize, uint8 Format, uint32 NumMips, uint32 Flags, FRHIResourceCreateInfo& CreateInfo);
+	virtual FTexture2DRHIRef RHICreateTexture2D_RenderThread(class FRHICommandListImmediate& RHICmdList, uint32 SizeX, uint32 SizeY, uint8 Format, uint32 NumMips, uint32 NumSamples, ETextureCreateFlags Flags, ERHIAccess InResourceState, FRHIResourceCreateInfo& CreateInfo) override;
+	virtual FTexture2DArrayRHIRef RHICreateTexture2DArray_RenderThread(class FRHICommandListImmediate& RHICmdList, uint32 SizeX, uint32 SizeY, uint32 SizeZ, uint8 Format, uint32 NumMips, uint32 NumSamples, ETextureCreateFlags Flags, ERHIAccess InResourceState, FRHIResourceCreateInfo& CreateInfo) override;
+	virtual FTexture3DRHIRef RHICreateTexture3D_RenderThread(class FRHICommandListImmediate& RHICmdList, uint32 SizeX, uint32 SizeY, uint32 SizeZ, uint8 Format, uint32 NumMips, ETextureCreateFlags Flags, ERHIAccess InResourceState, FRHIResourceCreateInfo& CreateInfo) override;
+	virtual FUnorderedAccessViewRHIRef RHICreateUnorderedAccessView_RenderThread(class FRHICommandListImmediate& RHICmdList, FRHIStructuredBuffer* StructuredBuffer, bool bUseUAVCounter, bool bAppendBuffer) override final;
+	virtual FUnorderedAccessViewRHIRef RHICreateUnorderedAccessView_RenderThread(class FRHICommandListImmediate& RHICmdList, FRHITexture* Texture, uint32 MipLevel) override final;
+	virtual FUnorderedAccessViewRHIRef RHICreateUnorderedAccessView_RenderThread(class FRHICommandListImmediate& RHICmdList, FRHIVertexBuffer* VertexBuffer, uint8 Format) override final;
+	virtual FShaderResourceViewRHIRef RHICreateShaderResourceView_RenderThread(class FRHICommandListImmediate& RHICmdList, FRHITexture* Texture, const FRHITextureSRVCreateInfo& CreateInfo) override final;
+	virtual FShaderResourceViewRHIRef RHICreateShaderResourceView_RenderThread(class FRHICommandListImmediate& RHICmdList, FRHIVertexBuffer* VertexBuffer, uint32 Stride, uint8 Format) override final;
+	virtual FShaderResourceViewRHIRef RHICreateShaderResourceView_RenderThread(class FRHICommandListImmediate& RHICmdList, const FShaderResourceViewInitializer& Initializer) override final;
+	virtual FShaderResourceViewRHIRef RHICreateShaderResourceViewWriteMask_RenderThread(class FRHICommandListImmediate& RHICmdList, FRHITexture2D* Texture2D) override final;
+
+	virtual FShaderResourceViewRHIRef CreateShaderResourceView_RenderThread(class FRHICommandListImmediate& RHICmdList, FRHIVertexBuffer* VertexBuffer, uint32 Stride, uint8 Format) override final;
+	virtual FShaderResourceViewRHIRef CreateShaderResourceView_RenderThread(class FRHICommandListImmediate& RHICmdList, const FShaderResourceViewInitializer& Initializer) override final;
+	virtual FShaderResourceViewRHIRef RHICreateShaderResourceView_RenderThread(class FRHICommandListImmediate& RHICmdList, FRHIStructuredBuffer* StructuredBuffer) override final;
+	virtual FTextureCubeRHIRef RHICreateTextureCube_RenderThread(class FRHICommandListImmediate& RHICmdList, uint32 Size, uint8 Format, uint32 NumMips, ETextureCreateFlags Flags, ERHIAccess InResourceState, FRHIResourceCreateInfo& CreateInfo) override;
+	virtual FTextureCubeRHIRef RHICreateTextureCubeArray_RenderThread(class FRHICommandListImmediate& RHICmdList, uint32 Size, uint32 ArraySize, uint8 Format, uint32 NumMips, ETextureCreateFlags Flags, ERHIAccess InResourceState, FRHIResourceCreateInfo& CreateInfo) override;
 
 	virtual FRenderQueryRHIRef RHICreateRenderQuery_RenderThread(class FRHICommandListImmediate& RHICmdList, ERenderQueryType QueryType) override final
 	{
@@ -910,11 +925,11 @@ public:
 public:
 
 #if	PLATFORM_SUPPORTS_VIRTUAL_TEXTURES
-	virtual void* CreateVirtualTexture(uint32 Flags, D3D12_RESOURCE_DESC& ResourceDesc, const struct FD3D12TextureLayout& TextureLayout, FD3D12Resource** ppResource, FPlatformMemory::FPlatformVirtualMemoryBlock& RawTextureBlock, D3D12_RESOURCE_STATES InitialUsage = D3D12_RESOURCE_STATE_PIXEL_SHADER_RESOURCE) = 0;
-	virtual void DestroyVirtualTexture(uint32 Flags, void* RawTextureMemory, FPlatformMemory::FPlatformVirtualMemoryBlock& RawTextureBlock, uint64 CommittedTextureSize) = 0;
+	virtual void* CreateVirtualTexture(ETextureCreateFlags InFlags, D3D12_RESOURCE_DESC& ResourceDesc, const struct FD3D12TextureLayout& TextureLayout, FD3D12Resource** ppResource, FPlatformMemory::FPlatformVirtualMemoryBlock& RawTextureBlock, D3D12_RESOURCE_STATES InitialUsage = D3D12_RESOURCE_STATE_PIXEL_SHADER_RESOURCE) = 0;
+	virtual void DestroyVirtualTexture(ETextureCreateFlags InFlags, void* RawTextureMemory, FPlatformMemory::FPlatformVirtualMemoryBlock& RawTextureBlock, uint64 CommittedTextureSize) = 0;
 #endif
-	virtual bool HandleSpecialLock(void*& MemoryOut, uint32 MipIndex, uint32 ArrayIndex, uint32 Flags, EResourceLockMode LockMode, const FD3D12TextureLayout& TextureLayout, void* RawTextureMemory, uint32& DestStride) { return false; }
-	virtual bool HandleSpecialUnlock(FRHICommandListBase* RHICmdList, uint32 MipIndex, uint32 Flags, const struct FD3D12TextureLayout& TextureLayout, void* RawTextureMemory) { return false; }
+	virtual bool HandleSpecialLock(void*& MemoryOut, uint32 MipIndex, uint32 ArrayIndex, uint32 InFlags, EResourceLockMode LockMode, const FD3D12TextureLayout& TextureLayout, void* RawTextureMemory, uint32& DestStride) { return false; }
+	virtual bool HandleSpecialUnlock(FRHICommandListBase* RHICmdList, uint32 MipIndex, uint32 InFlags, const struct FD3D12TextureLayout& TextureLayout, void* RawTextureMemory) { return false; }
 
 	FD3D12Adapter& GetAdapter(uint32_t Index = 0) { return *ChosenAdapters[Index]; }
 	int32 GetNumAdapters() const { return ChosenAdapters.Num(); }
@@ -965,9 +980,12 @@ protected:
 
 	template<typename BaseResourceType>
 	TD3D12Texture2D<BaseResourceType>* CreateD3D12Texture2D(class FRHICommandListImmediate* RHICmdList, uint32 SizeX, uint32 SizeY, uint32 SizeZ, bool bTextureArray, bool CubeTexture, EPixelFormat Format,
-		uint32 NumMips, uint32 NumSamples, uint32 Flags, FRHIResourceCreateInfo& CreateInfo);
+		uint32 NumMips, uint32 NumSamples, ETextureCreateFlags InFlags, FRHIResourceCreateInfo& CreateInfo);
+
+	FD3D12Texture3D* CreateD3D12Texture3D(class FRHICommandListImmediate* RHICmdList, uint32 SizeX, uint32 SizeY, uint32 SizeZ, EPixelFormat Format, uint32 NumMips, ETextureCreateFlags InFlags, FRHIResourceCreateInfo& CreateInfo);
 
-	FD3D12Texture3D* CreateD3D12Texture3D(class FRHICommandListImmediate* RHICmdList, uint32 SizeX, uint32 SizeY, uint32 SizeZ, EPixelFormat Format, uint32 NumMips, uint32 Flags, FRHIResourceCreateInfo& CreateInfo);
+	template<typename BaseResourceType>
+	TD3D12Texture2D<BaseResourceType>* CreateTextureFromResource(bool bTextureArray, bool bCubeTexture, EPixelFormat Format, ETextureCreateFlags TexCreateFlags, const FClearValueBinding& ClearValueBinding, ID3D12Resource* Resource);
 
 	template<typename BaseResourceType>
 	TD3D12Texture2D<BaseResourceType>* CreateAliasedD3D12Texture2D(TD3D12Texture2D<BaseResourceType>* SourceTexture);
@@ -990,7 +1008,7 @@ protected:
 	* @param StagingRectOUT - parameter is filled with the rectangle to read from the returned texture.
 	* @return The CPU readable Texture object.
 	*/
-	TRefCountPtr<FD3D12Resource> GetStagingTexture(FRHITexture* TextureRHI, FIntRect InRect, FIntRect& OutRect, FReadSurfaceDataFlags InFlags, D3D12_PLACED_SUBRESOURCE_FOOTPRINT &readBackHeapDesc);
+	TRefCountPtr<FD3D12Resource> GetStagingTexture(FRHITexture* TextureRHI, FIntRect InRect, FIntRect& OutRect, FReadSurfaceDataFlags InFlags, D3D12_PLACED_SUBRESOURCE_FOOTPRINT &readBackHeapDesc, uint32 GPUIndex);
 
 	void ReadSurfaceDataNoMSAARaw(FRHITexture* TextureRHI, FIntRect Rect, TArray<uint8>& OutData, FReadSurfaceDataFlags InFlags);
 
@@ -999,10 +1017,9 @@ protected:
 	void SetupRecursiveResources();
 
 	// This should only be called by Dynamic RHI member functions
-	inline FD3D12Device* GetRHIDevice()
+	inline FD3D12Device* GetRHIDevice(uint32 GPUIndex)
 	{
-		 // Multi-GPU support : any code using this function needs validation.
-		return GetAdapter().GetDevice(0);
+		return GetAdapter().GetDevice(GPUIndex);
 	}
 
 	HANDLE FlipEvent;
