Fix up deprecated code VariantManager color properties

diff --git a/Engine/Plugins/Enterprise/VariantManager/Source/VariantManager/Private/DisplayNodes/VariantManagerColorPropertyNode.cpp b/Engine/Plugins/Enterprise/VariantManager/Source/VariantManager/Private/DisplayNodes/VariantManagerColorPropertyNode.cpp
--- a/Engine/Plugins/Enterprise/VariantManager/Source/VariantManager/Private/DisplayNodes/VariantManagerColorPropertyNode.cpp
+++ b/Engine/Plugins/Enterprise/VariantManager/Source/VariantManager/Private/DisplayNodes/VariantManagerColorPropertyNode.cpp
@@ -196,7 +196,7 @@ TSharedPtr<SWidget> FVariantManagerColorPropertyNode::GetPropertyValueWidget()
 		return GetMultipleValuesWidget();
 	}
 
-	// Taken from FColorStructCustomization::CreateColorWidget
+	// Taken from FColorStructCustomization::CreateColorWidget.  Todo this should be combined into one color block based on changes to FColorStructCustomization
 	return	SNew(SBox)
 			.Padding(FMargin(4.0f, 0.0f, 4.0f, 0.0f)) // Extra padding to match the internal padding of single property nodes
 			[
@@ -209,7 +209,6 @@ TSharedPtr<SWidget> FVariantManagerColorPropertyNode::GetPropertyValueWidget()
 					SNew(SColorBlock)
 					.Color_Lambda([&](){ return CachedColor; })
 					.ShowBackgroundForAlpha(true)
-					.IgnoreAlpha(false)
 					.OnMouseButtonDown(this, &FVariantManagerColorPropertyNode::OnClickColorBlock)
 					.Size(FVector2D(35.0f, 17.0f))
 					.IsEnabled(true)
@@ -222,7 +221,7 @@ TSharedPtr<SWidget> FVariantManagerColorPropertyNode::GetPropertyValueWidget()
 					SNew(SColorBlock)
 					.Color_Lambda([&](){ return CachedColor; })
 					.ShowBackgroundForAlpha(false)
-					.IgnoreAlpha(true)
+					.AlphaDisplayMode(EColorBlockAlphaDisplayMode::Ignore)
 					.OnMouseButtonDown(this, &FVariantManagerColorPropertyNode::OnClickColorBlock)
 					.Size(FVector2D(35.0f, 17.0f))
 				]
