VariantManager: Update some of the UI to UE5 style;

Use the Layers icon for VariantManager temporarily;
Realign variant and variant set thumbnails;
Fix font size for struct property value editors;
Fix autocapture icon not showing up;
Update USD stage editor icon to the USD icon;

diff --git a/Engine/Plugins/Enterprise/VariantManager/Source/VariantManager/Private/DisplayNodes/VariantManagerStructPropertyNode.cpp b/Engine/Plugins/Enterprise/VariantManager/Source/VariantManager/Private/DisplayNodes/VariantManagerStructPropertyNode.cpp
--- a/Engine/Plugins/Enterprise/VariantManager/Source/VariantManager/Private/DisplayNodes/VariantManagerStructPropertyNode.cpp
+++ b/Engine/Plugins/Enterprise/VariantManager/Source/VariantManager/Private/DisplayNodes/VariantManagerStructPropertyNode.cpp
@@ -31,26 +31,22 @@ TSharedRef<SWidget> GetLabel(const FNumericProperty* InProp)
 	static FString BName = FString(TEXT("B"));
 	static FString ZName = FString(TEXT("Z"));
 
-	const FLinearColor* Foreground = &FLinearColor::Black;
 	const FLinearColor* Background = &FLinearColor::Transparent;
 
 	if (PropName == RName || PropName == XName)
 	{
-		Foreground = &FLinearColor::White;
 		Background = &SNumericEntryBox<uint8>::RedLabelBackgroundColor;
 	}
 	else if (PropName == GName || PropName == YName)
 	{
-		Foreground = &FLinearColor::White;
 		Background = &SNumericEntryBox<uint8>::GreenLabelBackgroundColor;
 	}
 	else if (PropName == BName || PropName == ZName)
 	{
-		Foreground = &FLinearColor::White;
 		Background = &SNumericEntryBox<uint8>::BlueLabelBackgroundColor;
 	}
 
-	return SNumericEntryBox<int32>::BuildLabel(FText::FromString(InProp->GetName()), *Foreground, *Background);
+	return SNumericEntryBox<int32>::BuildNarrowColorLabel( *Background );
 }
 
 template<typename NumericType>
@@ -141,7 +137,7 @@ TSharedRef<SWidget> FVariantManagerStructPropertyNode::GenerateFloatEntryBox(FNu
 
 	return SNew(SNumericEntryBox<double>)
 	.AllowSpin(true)
-	.LabelPadding(0)
+	.Font( FEditorStyle::GetFontStyle( "PropertyWindow.NormalFont" ) )
 	.OnBeginSliderMovement(this, &FVariantManagerStructPropertyNode::OnBeginSliderMovement, Prop)
 	.OnEndSliderMovement(this, &FVariantManagerStructPropertyNode::OnFloatEndSliderMovement, Prop, Offset)
 	.OnValueChanged(this, &FVariantManagerStructPropertyNode::OnFloatValueChanged, Prop)
@@ -157,6 +153,8 @@ TSharedRef<SWidget> FVariantManagerStructPropertyNode::GenerateFloatEntryBox(FNu
 	.MaxSliderValue(GET_OR_EMPTY(SliderMaxValue, double))
 	.SliderExponent(SliderExponent)
 	.Delta(Delta)
+	.LabelPadding( FMargin( 3 ) )
+	.LabelLocation( SNumericEntryBox<double>::ELabelLocation::Inside )
 	.Label()
 	[
 		GetLabel(Prop)
@@ -176,7 +174,7 @@ TSharedRef<SWidget> FVariantManagerStructPropertyNode::GenerateSignedEntryBox(FN
 
 	return SNew(SNumericEntryBox<int64>)
 	.AllowSpin(true)
-	.LabelPadding(0)
+	.Font( FEditorStyle::GetFontStyle( "PropertyWindow.NormalFont" ) )
 	.OnBeginSliderMovement(this, &FVariantManagerStructPropertyNode::OnBeginSliderMovement, Prop)
 	.OnEndSliderMovement(this, &FVariantManagerStructPropertyNode::OnSignedEndSliderMovement, Prop, Offset)
 	.OnValueChanged(this, &FVariantManagerStructPropertyNode::OnSignedValueChanged, Prop)
@@ -192,6 +190,8 @@ TSharedRef<SWidget> FVariantManagerStructPropertyNode::GenerateSignedEntryBox(FN
 	.MaxSliderValue(GET_OR_EMPTY(SliderMaxValue, int64))
 	.SliderExponent(SliderExponent)
 	.Delta(Delta)
+	.LabelPadding( FMargin( 3 ) )
+	.LabelLocation( SNumericEntryBox<int64>::ELabelLocation::Inside )
 	.Label()
 	[
 		GetLabel(Prop)
@@ -211,7 +211,7 @@ TSharedRef<SWidget> FVariantManagerStructPropertyNode::GenerateUnsignedEntryBox(
 
 	return SNew(SNumericEntryBox<uint64>)
 	.AllowSpin(true)
-	.LabelPadding(0)
+	.Font( FEditorStyle::GetFontStyle( "PropertyWindow.NormalFont" ) )
 	.OnBeginSliderMovement(this, &FVariantManagerStructPropertyNode::OnBeginSliderMovement, Prop)
 	.OnEndSliderMovement(this, &FVariantManagerStructPropertyNode::OnUnsignedEndSliderMovement, Prop, Offset)
 	.OnValueChanged(this, &FVariantManagerStructPropertyNode::OnUnsignedValueChanged, Prop)
@@ -227,6 +227,8 @@ TSharedRef<SWidget> FVariantManagerStructPropertyNode::GenerateUnsignedEntryBox(
 	.MaxSliderValue(GET_OR_EMPTY(SliderMaxValue, uint64))
 	.SliderExponent(SliderExponent)
 	.Delta(Delta)
+	.LabelPadding( FMargin( 3 ) )
+	.LabelLocation( SNumericEntryBox<uint64>::ELabelLocation::Inside )
 	.Label()
 	[
 		GetLabel(Prop)
@@ -300,10 +302,10 @@ TSharedPtr<SWidget> FVariantManagerStructPropertyNode::GetPropertyValueWidget()
 	}
 
 	TSharedPtr<SHorizontalBox> HorizBox = SNew(SHorizontalBox);
-	FMargin CommonMargin = FMargin(0.0f, 5.0f, 3.0f, 4.0f);
+	FMargin CommonMargin = FMargin(0.0f, 0.0f, 3.0f, 0.0f);
 
 	FMargin LastMargin = CommonMargin;
-	LastMargin.Right = 0.0f;
+	LastMargin.Right = 4.0f;
 
 	FNumericProperty* LastProp = nullptr;
 	for (TFieldIterator<FNumericProperty> NumPropIter(StructClassToEdit); NumPropIter; ++NumPropIter)
@@ -442,7 +444,7 @@ TSharedPtr<SWidget> FVariantManagerStructPropertyNode::GetPropertyValueWidget()
 	}
 
 	return	SNew(SBox)
-			.Padding(FMargin(4.0f, 0.0f, 4.0f, 0.0f)) // Extra padding to match the internal padding of single property nodes
+			.Padding(FMargin(4.0f, 0.0f, 0.0f, 0.0f)) // Extra padding to match the internal padding of single property nodes
 			[
 				HorizBox.ToSharedRef()
 			];
diff --git a/Engine/Plugins/Enterprise/VariantManager/Source/VariantManager/Private/DisplayNodes/VariantManagerVariantNode.cpp b/Engine/Plugins/Enterprise/VariantManager/Source/VariantManager/Private/DisplayNodes/VariantManagerVariantNode.cpp
--- a/Engine/Plugins/Enterprise/VariantManager/Source/VariantManager/Private/DisplayNodes/VariantManagerVariantNode.cpp
+++ b/Engine/Plugins/Enterprise/VariantManager/Source/VariantManager/Private/DisplayNodes/VariantManagerVariantNode.cpp
@@ -83,7 +83,7 @@ TSharedRef<SWidget> FVariantManagerVariantNode::GetCustomOutlinerContent(TShared
 		.HAlign(HAlign_Fill)
 		.BorderImage(this, &FVariantManagerDisplayNode::GetNodeBorderImage)
 		.BorderBackgroundColor(this, &FVariantManagerDisplayNode::GetNodeBackgroundTint)
-		.Padding(FMargin(20.0f, 0.0f, 2.0f, 0.0f))
+		.Padding(FMargin(26.0f, 0.0f, 2.0f, 0.0f))
 		[
 			SNew(SHorizontalBox)
 
diff --git a/Engine/Plugins/Enterprise/VariantManager/Source/VariantManager/Private/SVariantManager.cpp b/Engine/Plugins/Enterprise/VariantManager/Source/VariantManager/Private/SVariantManager.cpp
--- a/Engine/Plugins/Enterprise/VariantManager/Source/VariantManager/Private/SVariantManager.cpp
+++ b/Engine/Plugins/Enterprise/VariantManager/Source/VariantManager/Private/SVariantManager.cpp
@@ -243,8 +243,6 @@ void SVariantManager::Construct(const FArguments& InArgs, TSharedRef<FVariantMan
 		OnMapChangedHandle = LevelEditorModule.OnMapChanged().AddSP(this, &SVariantManager::OnMapChanged);
 	}
 
-	RecordButtonBrush = MakeShared<FSlateImageBrush>(FPaths::EngineContentDir() / TEXT("Editor/Slate/Icons/CA_Record.png"), FVector2D(24.0f, 24.0f));
-
 	RightTreeRootItems.Empty();
 	RightTreeRootItems.Reserve(2);
 	RightTreeRootItems.Add(MakeShared<ERightTreeRowType>(ERightTreeRowType::PropertiesHeader));
@@ -318,30 +316,18 @@ void SVariantManager::Construct(const FArguments& InArgs, TSharedRef<FVariantMan
 				.AutoWidth()
 				.MaxWidth(VM_COMMON_HEADER_MAX_HEIGHT) // square aspect ratio
 				[
-					SNew(SBox)
-					.HeightOverride(VM_COMMON_HEADER_MAX_HEIGHT - 8.0f) // These so that it matches the height of the search box
-					.WidthOverride(VM_COMMON_HEADER_MAX_HEIGHT - 8.0f)
-					[
-						SNew(SCheckBox)
-						.Style(FCoreStyle::Get(), "ToggleButtonCheckbox")
-						.ToolTipText(LOCTEXT("AutoCaptureTooltip", "Enable or disable auto-capturing properties"))
-						.IsChecked_Lambda([&bAutoCaptureProperties = bAutoCaptureProperties]()
-						{
-							return bAutoCaptureProperties? ECheckBoxState::Checked : ECheckBoxState::Unchecked;
-						})
-						.OnCheckStateChanged_Lambda([&bAutoCaptureProperties = bAutoCaptureProperties](const ECheckBoxState NewState)
-						{
-							bAutoCaptureProperties = NewState == ECheckBoxState::Checked;
-						})
-						[
-							SNew(SBox)
-							.Padding(FMargin(0.0f, 2.0, 2.0f, 2.0)) // Extra padding on the right because ToggleButtonCheckboxes always nudges the image to the right
-							[
-								SNew(SImage)
-								.Image(RecordButtonBrush.Get())
-							]
-						]
-					]
+					SNew(SCheckBox)
+					.Style( FVariantManagerStyle::Get(), "AutoCaptureCheckbox" )
+					.Padding( FMargin( 10.f, 10.f, 10.0f, 10.f ) ) // Give some space to show up the style image or else it will collapse to nothing since it has no content
+					.ToolTipText(LOCTEXT("AutoCaptureTooltip", "Enable or disable auto-capturing properties"))
+					.IsChecked_Lambda([&bAutoCaptureProperties = bAutoCaptureProperties]()
+					{
+						return bAutoCaptureProperties? ECheckBoxState::Checked : ECheckBoxState::Unchecked;
+					})
+					.OnCheckStateChanged_Lambda([&bAutoCaptureProperties = bAutoCaptureProperties](const ECheckBoxState NewState)
+					{
+						bAutoCaptureProperties = NewState == ECheckBoxState::Checked;
+					})
 				]
 
 				+SHorizontalBox::Slot()
diff --git a/Engine/Plugins/Enterprise/VariantManager/Source/VariantManager/Private/SVariantManager.h b/Engine/Plugins/Enterprise/VariantManager/Source/VariantManager/Private/SVariantManager.h
--- a/Engine/Plugins/Enterprise/VariantManager/Source/VariantManager/Private/SVariantManager.h
+++ b/Engine/Plugins/Enterprise/VariantManager/Source/VariantManager/Private/SVariantManager.h
@@ -313,9 +313,6 @@ private:
 
 	TArray<TSharedRef<ERightTreeRowType>> RightTreeRootItems;
 
-	// TODO: Make separate VariantManagerStyle
-	TSharedPtr<FSlateImageBrush> RecordButtonBrush;
-
 	struct FCachedPropertyPath
 	{
 		UObject* Object;
diff --git a/Engine/Plugins/Enterprise/VariantManager/Source/VariantManager/Private/VariantManagerModule.cpp b/Engine/Plugins/Enterprise/VariantManager/Source/VariantManager/Private/VariantManagerModule.cpp
--- a/Engine/Plugins/Enterprise/VariantManager/Source/VariantManager/Private/VariantManagerModule.cpp
+++ b/Engine/Plugins/Enterprise/VariantManager/Source/VariantManager/Private/VariantManagerModule.cpp
@@ -89,10 +89,12 @@ public:
 			UnregisterTabSpawner( TabManager );
 		}
 
+		const FSlateIcon LayersIcon( FEditorStyle::GetStyleSetName(), "LevelEditor.Tabs.Layers" );
+
 		TabManager->RegisterTabSpawner( FLevelVariantSetsEditorToolkit::GetVariantManagerTabID(), FOnSpawnTab::CreateStatic( &FVariantManagerModule::HandleTabManagerSpawnTab ) )
 			.SetDisplayName( LOCTEXT("VariantManagerMainTab", "Variant Manager") )
 			.SetGroup( WorkspaceMenu::GetMenuStructure().GetLevelEditorCategory() )
-			.SetIcon( FSlateIcon(FEditorStyle::GetStyleSetName(), "LevelEditor.GameSettings.Small") );
+			.SetIcon( LayersIcon );
 	}
 
 	static void UnregisterTabSpawner( const TSharedPtr< FTabManager >& TabManager )
diff --git a/Engine/Plugins/Enterprise/VariantManager/Source/VariantManager/Private/VariantManagerStyle.cpp b/Engine/Plugins/Enterprise/VariantManager/Source/VariantManager/Private/VariantManagerStyle.cpp
--- a/Engine/Plugins/Enterprise/VariantManager/Source/VariantManager/Private/VariantManagerStyle.cpp
+++ b/Engine/Plugins/Enterprise/VariantManager/Source/VariantManager/Private/VariantManagerStyle.cpp
@@ -64,6 +64,17 @@ void FVariantManagerStyle::Initialize()
 
 		/** The amount to indent child nodes of the layout tree */
 		StyleSet->Set( "VariantManager.Spacings.IndentAmount", 10.0f );
+
+		/** Style used for the auto-capture button */
+		const FCheckBoxStyle Style = FCheckBoxStyle()
+			.SetCheckBoxType( ESlateCheckBoxType::ToggleButton )
+			.SetUncheckedImage( BOX_BRUSH( "/Sequencer/Transport_Bar/Record_24x_OFF", 0.0f, FSlateColor( FLinearColor::Gray ) ) )
+			.SetUncheckedHoveredImage( BOX_BRUSH( "/Sequencer/Transport_Bar/Record_24x_OFF", 0.0f, FSlateColor( FLinearColor::White ) ) )
+			.SetUncheckedPressedImage( BOX_BRUSH( "/Sequencer/Transport_Bar/Record_24x_OFF", 0.0f, FSlateColor( FLinearColor(0.75f, 0.75f, 0.75f) ) ) )
+			.SetCheckedImage( BOX_BRUSH( "/Sequencer/Transport_Bar/Recording_24x", 0.0f, FSlateColor( FLinearColor::Gray ) ) )
+			.SetCheckedHoveredImage( BOX_BRUSH( "/Sequencer/Transport_Bar/Recording_24x", 0.0f, FSlateColor( FLinearColor::White ) ) )
+			.SetCheckedPressedImage( BOX_BRUSH( "/Sequencer/Transport_Bar/Recording_24x", 0.0f, FSlateColor( FLinearColor( 0.75f, 0.75f, 0.75f ) ) ) );
+		StyleSet->Set( "AutoCaptureCheckbox", Style );
 	}
 
 	FSlateStyleRegistry::RegisterSlateStyle(*StyleSet.Get());
diff --git a/Engine/Plugins/Importers/USDImporter/Source/USDStageEditor/Private/USDStageEditorModule.cpp b/Engine/Plugins/Importers/USDImporter/Source/USDStageEditor/Private/USDStageEditorModule.cpp
--- a/Engine/Plugins/Importers/USDImporter/Source/USDStageEditor/Private/USDStageEditorModule.cpp
+++ b/Engine/Plugins/Importers/USDImporter/Source/USDStageEditor/Private/USDStageEditorModule.cpp
@@ -27,7 +27,7 @@ public:
 	{
 		return SNew(SDockTab)
 			.TabRole(ETabRole::NomadTab)
-			.Icon( FEditorStyle::GetBrush("LevelEditor.Tabs.Layers") )
+			.Icon( FEditorStyle::GetBrush("LevelEditor.Tabs.USDStage") )
 			.Label( LOCTEXT( "USDStage", "USD Stage" ) )
 			[
 				SNew(SBorder)
