Variant Manager: Fix obvious broken UI/UX

diff --git a/Engine/Plugins/Enterprise/VariantManager/Source/VariantManager/Private/DisplayNodes/VariantManagerPropertyNode.cpp b/Engine/Plugins/Enterprise/VariantManager/Source/VariantManager/Private/DisplayNodes/VariantManagerPropertyNode.cpp
--- a/Engine/Plugins/Enterprise/VariantManager/Source/VariantManager/Private/DisplayNodes/VariantManagerPropertyNode.cpp
+++ b/Engine/Plugins/Enterprise/VariantManager/Source/VariantManager/Private/DisplayNodes/VariantManagerPropertyNode.cpp
@@ -360,7 +360,7 @@ TSharedPtr<SWidget> FVariantManagerPropertyNode::GetPropertyValueWidget()
 	}
 	else
 	{
-	FMemory::Memcpy((uint8*)SinglePropWidgetDataPtr, FirstRecordedData.GetData(), FirstPropertyValue->GetValueSizeInBytes());
+		FMemory::Memcpy((uint8*)SinglePropWidgetDataPtr, FirstRecordedData.GetData(), FirstPropertyValue->GetValueSizeInBytes());
 	}
 
 	// Update recorded data when user modifies the widget (modifying the widget will modify the
@@ -420,7 +420,6 @@ TSharedRef<SWidget> FVariantManagerPropertyNode::GetCustomOutlinerContent(TShare
 				[
 					SNew(SBox)
 					.Padding(FMargin( 3.0, 0.0f, 3.0f, 0.0f ))
-					.HeightOverride(26)
 					[
 						GetPropertyValueWidget().ToSharedRef()
 					]
@@ -429,17 +428,15 @@ TSharedRef<SWidget> FVariantManagerPropertyNode::GetCustomOutlinerContent(TShare
 				+SHorizontalBox::Slot()
 				.Padding( FMargin(0.0f, 0.0f, 3.0f, 0.0) )
 				.AutoWidth()
-				.VAlign( VAlign_Center )
+				.VAlign( VAlign_Top )
 				[
 					SNew(SBox)
-					.HeightOverride(21)
-					.WidthOverride(21)
 					[
 						SAssignNew(RecordButton, SButton)
 						.IsFocusable(false)
 						.ToolTipText(LOCTEXT("UseCurrentTooltip", "Record the current value for this property"))
 						.ButtonStyle( FEditorStyle::Get(), "HoverHintOnly" )
-						.ContentPadding(0.0f)
+						.ContentPadding(6.0f)
 						.OnClicked(this, &FVariantManagerPropertyNode::RecordMultipleValues)
 						.Visibility(FVariantManagerPropertyNode::GetRecordButtonVisibility())
 						.Content()
@@ -457,19 +454,17 @@ TSharedRef<SWidget> FVariantManagerPropertyNode::GetCustomOutlinerContent(TShare
 				]
 
 				+SHorizontalBox::Slot()
-				.Padding( FMargin(0.0f, 0.0f, 1.0f, 0.0) )
+				.Padding( FMargin(0.0f, 0.0f, 3.0f, 0.0) )
 				.AutoWidth()
-				.VAlign( VAlign_Center )
+				.VAlign( VAlign_Top )
 				[
 					SNew(SBox)
-					.HeightOverride(21)
-					.WidthOverride(21)
 					[
 						SAssignNew(ResetButton, SButton)
 						.IsFocusable(false)
 						.ToolTipText(LOCTEXT("ResetTooltip", "Reset to the property's default value"))
 						.ButtonStyle( FEditorStyle::Get(), "HoverHintOnly" )
-						.ContentPadding(6.0f)
+						.ContentPadding(FMargin(3.0f, 8.0f))
 						.OnClicked(this, &FVariantManagerPropertyNode::ResetMultipleValuesToDefault)
 						.Visibility(FVariantManagerPropertyNode::GetResetButtonVisibility())
 						.Content()
