VariantManager: Align UI with new UE5 template

diff --git a/Engine/Plugins/Enterprise/VariantManager/Source/VariantManager/Private/DisplayNodes/VariantManagerActorNode.cpp b/Engine/Plugins/Enterprise/VariantManager/Source/VariantManager/Private/DisplayNodes/VariantManagerActorNode.cpp
--- a/Engine/Plugins/Enterprise/VariantManager/Source/VariantManager/Private/DisplayNodes/VariantManagerActorNode.cpp
+++ b/Engine/Plugins/Enterprise/VariantManager/Source/VariantManager/Private/DisplayNodes/VariantManagerActorNode.cpp
@@ -579,7 +579,6 @@ TSharedRef<SWidget> FVariantManagerActorNode::GetCustomOutlinerContent(TSharedPt
 		SNew(SBorder)
 		.VAlign(VAlign_Center)
 		.BorderImage(this, &FVariantManagerDisplayNode::GetNodeBorderImage)
-		.BorderBackgroundColor(this, &FVariantManagerDisplayNode::GetNodeBackgroundTint)
 		.Padding(FMargin(2.0f, 0.0f, 2.0f, 0.0f))
 		[
 			SNew(SBox)
diff --git a/Engine/Plugins/Enterprise/VariantManager/Source/VariantManager/Private/DisplayNodes/VariantManagerDisplayNode.cpp b/Engine/Plugins/Enterprise/VariantManager/Source/VariantManager/Private/DisplayNodes/VariantManagerDisplayNode.cpp
--- a/Engine/Plugins/Enterprise/VariantManager/Source/VariantManager/Private/DisplayNodes/VariantManagerDisplayNode.cpp
+++ b/Engine/Plugins/Enterprise/VariantManager/Source/VariantManager/Private/DisplayNodes/VariantManagerDisplayNode.cpp
@@ -20,7 +20,12 @@ FVariantManagerDisplayNode::FVariantManagerDisplayNode(TSharedPtr<FVariantManage
 	, bExpanded( false )
 	, bSelected( false )
 {
-	BackgroundBrush = FAppStyle::GetBrush("Sequencer.AnimationOutliner.DefaultBorder");
+	BackgroundBrush = nullptr;
+}
+
+const FTableRowStyle* FVariantManagerDisplayNode::GetRowStyle() const
+{
+	return &FAppStyle::Get().GetWidgetStyle<FTableRowStyle>("SimpleTableView.Row");
 }
 
 FSlateColor FVariantManagerDisplayNode::GetDisplayNameColor() const
@@ -54,8 +59,6 @@ TSharedRef<SWidget> FVariantManagerDisplayNode::GetCustomOutlinerContent(TShared
 		SNew(SBorder)
 		.VAlign(VAlign_Fill)
 		.HAlign(HAlign_Fill)
-		.BorderImage(this, &FVariantManagerDisplayNode::GetNodeBorderImage)
-		.BorderBackgroundColor(this, &FVariantManagerDisplayNode::GetNodeBackgroundTint)
 	];
 }
 
@@ -84,22 +87,6 @@ const FSlateBrush* FVariantManagerDisplayNode::GetNodeBorderImage() const
 	return BackgroundBrush;
 }
 
-FSlateColor FVariantManagerDisplayNode::GetNodeBackgroundTint() const
-{
-	if (IsSelected())
-	{
-		return FAppStyle::GetSlateColor("SelectionColor_Pressed");
-	}
-	else if (IsHovered())
-	{
-		return FLinearColor(FColor(72, 72, 72, 255));
-	}
-	else
-	{
-		return FLinearColor(FColor(62, 62, 62, 255));
-	}
-}
-
 void FVariantManagerDisplayNode::BuildContextMenu(FMenuBuilder& MenuBuilder)
 {
 
diff --git a/Engine/Plugins/Enterprise/VariantManager/Source/VariantManager/Private/DisplayNodes/VariantManagerDisplayNode.h b/Engine/Plugins/Enterprise/VariantManager/Source/VariantManager/Private/DisplayNodes/VariantManagerDisplayNode.h
--- a/Engine/Plugins/Enterprise/VariantManager/Source/VariantManager/Private/DisplayNodes/VariantManagerDisplayNode.h
+++ b/Engine/Plugins/Enterprise/VariantManager/Source/VariantManager/Private/DisplayNodes/VariantManagerDisplayNode.h
@@ -62,6 +62,7 @@ public:
 		return FText();
 	}
 
+	virtual const FTableRowStyle* GetRowStyle() const;
 	virtual FSlateColor GetDisplayNameColor() const;
 	virtual FText GetDisplayNameToolTipText() const;
 	virtual void SetDisplayName(const FText& NewDisplayName) {}
@@ -85,7 +86,6 @@ public:
 	virtual FSlateColor GetIconColor() const;
 	virtual FText GetIconToolTipText() const;
 	virtual const FSlateBrush* GetNodeBorderImage() const;
-	virtual FSlateColor GetNodeBackgroundTint() const;
 	virtual void BuildContextMenu(FMenuBuilder& MenuBuilder);
 
 	uint32 GetNumChildren() const
diff --git a/Engine/Plugins/Enterprise/VariantManager/Source/VariantManager/Private/DisplayNodes/VariantManagerPropertyNode.cpp b/Engine/Plugins/Enterprise/VariantManager/Source/VariantManager/Private/DisplayNodes/VariantManagerPropertyNode.cpp
--- a/Engine/Plugins/Enterprise/VariantManager/Source/VariantManager/Private/DisplayNodes/VariantManagerPropertyNode.cpp
+++ b/Engine/Plugins/Enterprise/VariantManager/Source/VariantManager/Private/DisplayNodes/VariantManagerPropertyNode.cpp
@@ -64,11 +64,6 @@ const FSlateBrush* FVariantManagerPropertyNode::GetIconOverlayBrush() const
 	return nullptr;
 }
 
-FSlateColor FVariantManagerPropertyNode::GetNodeBackgroundTint() const
-{
-	return FLinearColor(FColor(62, 62, 62, 255));
-}
-
 FText FVariantManagerPropertyNode::GetIconToolTipText() const
 {
 	return FText();
@@ -372,16 +367,15 @@ TSharedPtr<SWidget> FVariantManagerPropertyNode::GetPropertyValueWidget()
 
 TSharedRef<SWidget> FVariantManagerPropertyNode::GetCustomOutlinerContent(TSharedPtr<SVariantManagerTableRow> InTableRow)
 {
-	// Using this syncs all splitters between property nodes and also the header
+	// Using this syncs all splitters between property nodes
 	TSharedPtr<SVariantManager> VariantManagerWidget = GetVariantManager().Pin()->GetVariantManagerWidget();
-	FColumnSizeData& ColumnSizeData = VariantManagerWidget->GetPropertiesColumnSizeData();
+	FVariantManagerPropertiesColumnSizeData& ColumnSizeData = VariantManagerWidget->GetPropertiesColumnSizeData();
 
 	return SNew(SBox)
 	[
 		SNew(SBorder)
 		.VAlign(VAlign_Center)
 		.BorderImage(this, &FVariantManagerDisplayNode::GetNodeBorderImage)
-		.BorderBackgroundColor(this, &FVariantManagerDisplayNode::GetNodeBackgroundTint)
 		.Padding(FMargin(2.0f, 0.0f, 2.0f, 0.0f))
 		[
 			SNew(SSplitter)
@@ -391,8 +385,8 @@ TSharedRef<SWidget> FVariantManagerPropertyNode::GetCustomOutlinerContent(TShare
 			.Orientation(Orient_Horizontal)
 
 			+ SSplitter::Slot()
-			.Value(ColumnSizeData.MiddleColumnWidth)
-			.OnSlotResized(SSplitter::FOnSlotResized::CreateLambda([](float InNewWidth) {}))
+			.Value(ColumnSizeData.NameColumnWidth)
+			.OnSlotResized(ColumnSizeData.OnSplitterNameColumnChanged)
 			[
 				SNew(SBox)
 				.VAlign(VAlign_Center)
@@ -412,8 +406,8 @@ TSharedRef<SWidget> FVariantManagerPropertyNode::GetCustomOutlinerContent(TShare
 			]
 
 			+ SSplitter::Slot()
-			.Value(ColumnSizeData.MiddleColumnWidth)
-			.OnSlotResized(SSplitter::FOnSlotResized::CreateLambda([](float InNewWidth) {}))
+			.Value(ColumnSizeData.ValueColumnWidth)
+			.OnSlotResized(ColumnSizeData.OnSplitterValueColumnChanged)
 			[
 				SNew(SHorizontalBox)
 				+SHorizontalBox::Slot()
@@ -428,7 +422,7 @@ TSharedRef<SWidget> FVariantManagerPropertyNode::GetCustomOutlinerContent(TShare
 				+SHorizontalBox::Slot()
 				.Padding( FMargin(0.0f, 0.0f, 3.0f, 0.0) )
 				.AutoWidth()
-				.VAlign( VAlign_Top )
+				.VAlign( VAlign_Center )
 				[
 					SNew(SBox)
 					[
@@ -456,7 +450,7 @@ TSharedRef<SWidget> FVariantManagerPropertyNode::GetCustomOutlinerContent(TShare
 				+SHorizontalBox::Slot()
 				.Padding( FMargin(0.0f, 0.0f, 3.0f, 0.0) )
 				.AutoWidth()
-				.VAlign( VAlign_Top )
+				.VAlign(VAlign_Center)
 				[
 					SNew(SBox)
 					[
diff --git a/Engine/Plugins/Enterprise/VariantManager/Source/VariantManager/Private/DisplayNodes/VariantManagerPropertyNode.h b/Engine/Plugins/Enterprise/VariantManager/Source/VariantManager/Private/DisplayNodes/VariantManagerPropertyNode.h
--- a/Engine/Plugins/Enterprise/VariantManager/Source/VariantManager/Private/DisplayNodes/VariantManagerPropertyNode.h
+++ b/Engine/Plugins/Enterprise/VariantManager/Source/VariantManager/Private/DisplayNodes/VariantManagerPropertyNode.h
@@ -34,7 +34,6 @@ public:
 	virtual void BuildContextMenu(FMenuBuilder& MenuBuilder) override;
 	virtual FText GetDisplayNameToolTipText() const override;
 	virtual const FSlateBrush* GetIconOverlayBrush() const override;
-	virtual FSlateColor GetNodeBackgroundTint() const override;
 	virtual FText GetIconToolTipText() const override;
 	virtual EVariantManagerNodeType GetType() const override;
 	virtual bool IsReadOnly() const override;
diff --git a/Engine/Plugins/Enterprise/VariantManager/Source/VariantManager/Private/DisplayNodes/VariantManagerVariantNode.cpp b/Engine/Plugins/Enterprise/VariantManager/Source/VariantManager/Private/DisplayNodes/VariantManagerVariantNode.cpp
--- a/Engine/Plugins/Enterprise/VariantManager/Source/VariantManager/Private/DisplayNodes/VariantManagerVariantNode.cpp
+++ b/Engine/Plugins/Enterprise/VariantManager/Source/VariantManager/Private/DisplayNodes/VariantManagerVariantNode.cpp
@@ -9,6 +9,7 @@
 #include "VariantManagerEditorCommands.h"
 #include "VariantManagerNodeTree.h"
 #include "VariantManagerSelection.h"
+#include "VariantManagerStyle.h"
 #include "VariantObjectBinding.h"
 #include "VariantSet.h"
 
@@ -63,15 +64,12 @@ TSharedRef<SWidget> FVariantManagerVariantNode::GetCustomOutlinerContent(TShared
 		.ToolTipText(this, &FVariantManagerDisplayNode::GetDisplayNameToolTipText)
 		.Clipping(EWidgetClipping::ClipToBounds);
 
-
 	RadioButton = SNew(SCheckBox)
 		.HAlign(HAlign_Right)
 		.Padding(FMargin(0))
-		.Style(FAppStyle::Get(), "Menu.RadioButton")
-		.ForegroundColor(FSlateColor(FLinearColor(1.0f, 1.0f, 1.0f)))
+		.Style(FVariantManagerStyle::Get(), "VariantManager.VariantRadioButton")
 		.ToolTipText(LOCTEXT("ActivateVariantRadioToolTip", "Activate the variant"))
 		.IsChecked(FVariantManagerVariantNode::IsRadioButtonChecked())
-		.Visibility_Lambda([&](){return RadioButton->IsChecked() ? EVisibility::HitTestInvisible : EVisibility::Visible; }) // So we can't unclick it
 		.OnCheckStateChanged(this, &FVariantManagerVariantNode::OnRadioButtonStateChanged);
 
 	return
@@ -81,36 +79,35 @@ TSharedRef<SWidget> FVariantManagerVariantNode::GetCustomOutlinerContent(TShared
 		SNew(SBorder)
 		.VAlign(VAlign_Center)
 		.HAlign(HAlign_Fill)
-		.BorderImage(this, &FVariantManagerDisplayNode::GetNodeBorderImage)
-		.BorderBackgroundColor(this, &FVariantManagerDisplayNode::GetNodeBackgroundTint)
+		.BorderImage(nullptr)
 		.Padding(FMargin(26.0f, 0.0f, 2.0f, 0.0f))
 		[
 			SNew(SHorizontalBox)
 
 			+ SHorizontalBox::Slot()
-			.Padding(FMargin(4.f, 0.f, 4.f, 0.f))
 			.VAlign(VAlign_Center)
+			.HAlign(HAlign_Right)
 			.AutoWidth()
+			.Padding(FMargin(0.f, 0.f, 4.f, 0.f))
 			[
-				GetThumbnailWidget()
+				RadioButton.ToSharedRef()
 			]
 
 			+ SHorizontalBox::Slot()
+			.Padding(FMargin(4.f, 0.f, 4.f, 0.f))
 			.VAlign(VAlign_Center)
-			.HAlign(HAlign_Left)
-			.FillWidth(1.0f)
-			.Padding(FMargin(0.f, 0.f, 4.f, 0.f))
+			.AutoWidth()
 			[
-				EditableLabel.ToSharedRef()
+				GetThumbnailWidget()
 			]
 
 			+ SHorizontalBox::Slot()
 			.VAlign(VAlign_Center)
-			.HAlign(HAlign_Right)
-			.AutoWidth()
+			.HAlign(HAlign_Left)
+			.FillWidth(1.0f)
 			.Padding(FMargin(0.f, 0.f, 4.f, 0.f))
 			[
-				RadioButton.ToSharedRef()
+				EditableLabel.ToSharedRef()
 			]
 		]
 	];
diff --git a/Engine/Plugins/Enterprise/VariantManager/Source/VariantManager/Private/DisplayNodes/VariantManagerVariantSetNode.cpp b/Engine/Plugins/Enterprise/VariantManager/Source/VariantManager/Private/DisplayNodes/VariantManagerVariantSetNode.cpp
--- a/Engine/Plugins/Enterprise/VariantManager/Source/VariantManager/Private/DisplayNodes/VariantManagerVariantSetNode.cpp
+++ b/Engine/Plugins/Enterprise/VariantManager/Source/VariantManager/Private/DisplayNodes/VariantManagerVariantSetNode.cpp
@@ -30,6 +30,7 @@
 #include "Framework/Notifications/NotificationManager.h"
 #include "Input/DragAndDrop.h"
 #include "ScopedTransaction.h"
+#include "Styling/StyleColors.h"
 #include "Textures/SlateIcon.h"
 #include "Widgets/Input/SButton.h"
 #include "Widgets/Input/SButton.h"
@@ -44,13 +45,24 @@ FVariantManagerVariantSetNode::FVariantManagerVariantSetNode( UVariantSet& InVar
 {
 	ExpandedBackgroundBrush = FAppStyle::GetBrush("Sequencer.AnimationOutliner.TopLevelBorder_Expanded");
 	CollapsedBackgroundBrush = FAppStyle::GetBrush("Sequencer.AnimationOutliner.TopLevelBorder_Collapsed");
+	RowStyle = FAppStyle::Get().GetWidgetStyle<FTableRowStyle>("SimpleTableView.Row");
+	RowStyle
+		.SetEvenRowBackgroundBrush(FSlateRoundedBoxBrush(FStyleColors::Header, 2.f, FStyleColors::Transparent, 1.f))
+		.SetOddRowBackgroundBrush(FSlateRoundedBoxBrush(FStyleColors::Header, 2.f, FStyleColors::Transparent, 1.f))
+		.SetEvenRowBackgroundHoveredBrush(FSlateRoundedBoxBrush(FStyleColors::SelectHover, 2.f))
+		.SetOddRowBackgroundHoveredBrush(FSlateRoundedBoxBrush(FStyleColors::SelectHover, 2.f));
 
 	bExpanded = InVariantSet.IsExpanded();
 }
 
+const FTableRowStyle* FVariantManagerVariantSetNode::GetRowStyle() const
+{
+	return &RowStyle;
+}
+
 TSharedRef<SWidget> FVariantManagerVariantSetNode::GetCustomOutlinerContent(TSharedPtr<SVariantManagerTableRow> InTableRow)
 {
-	FSlateFontInfo NodeFont = FAppStyle::GetFontStyle("Sequencer.AnimationOutliner.RegularFont");
+	FSlateFontInfo NodeFont = FAppStyle::GetFontStyle("BoldFont");
 
 	EditableLabel = SNew(SInlineEditableTextBlock)
 		.IsReadOnly(this, &FVariantManagerDisplayNode::IsReadOnly)
@@ -74,18 +86,16 @@ TSharedRef<SWidget> FVariantManagerVariantSetNode::GetCustomOutlinerContent(TSha
 
 	return
 	SNew( SBorder )
-	.BorderImage( FAppStyle::GetBrush( "WhiteBrush" ) )
-	.BorderBackgroundColor( FVariantManagerStyle::Get()->GetColor("VariantManager.Panels.LightBackgroundColor") )
-	.Padding( FMargin( 0.0f, bIsFirstRowOfTree ? 0.0f : FVariantManagerStyle::Get()->GetFloat( "VariantManager.Spacings.BorderThickness" ), 0.0f, 0.0f ) )
+	.BorderImage(nullptr)
+	.Padding( FMargin( 0.0f, bIsFirstRowOfTree ? 0.0f : FVariantManagerStyle::Get().GetFloat( "VariantManager.Spacings.BorderThickness" ), 0.0f, 0.0f ) )
 	[
 		SNew(SBox)
 		.HeightOverride(78)
 		[
 			SNew(SBorder)
+			.BorderImage(nullptr)
 			.VAlign(VAlign_Center)
 			.HAlign(HAlign_Fill)
-			.BorderImage(this, &FVariantManagerDisplayNode::GetNodeBorderImage)
-			.BorderBackgroundColor(this, &FVariantManagerDisplayNode::GetNodeBackgroundTint)
 			.Padding(FMargin(2.0f, 0.0f, 2.0f, 0.0f))
 			[
 				SNew(SHorizontalBox)
@@ -96,7 +106,7 @@ TSharedRef<SWidget> FVariantManagerVariantSetNode::GetCustomOutlinerContent(TSha
 				.MaxWidth(24.0f)
 				.AutoWidth()
 				[
-					SNew(SExpanderArrow, InTableRow).IndentAmount( FVariantManagerStyle::Get()->GetFloat( "VariantManager.Spacings.IndentAmount" ) )
+					SNew(SExpanderArrow, InTableRow).IndentAmount( FVariantManagerStyle::Get().GetFloat( "VariantManager.Spacings.IndentAmount" ) )
 				]
 
 				+ SHorizontalBox::Slot()
@@ -124,8 +134,7 @@ TSharedRef<SWidget> FVariantManagerVariantSetNode::GetCustomOutlinerContent(TSha
 				.Padding(FMargin(4.f, 0.f, 4.f, 0.f))
 				[
 					SNew(SButton)
-					.ButtonStyle(FAppStyle::Get(), "NoBorder")
-					.ToolTipText(LOCTEXT("AddVariantToolTip", "Add a new variant"))
+					.ButtonStyle(FAppStyle::Get(), "SimpleButton")
 					.OnClicked_Lambda([this]
 					{
 						TSharedPtr<FVariantManager> VariantManagerPtr = GetVariantManager().Pin();
@@ -141,13 +150,12 @@ TSharedRef<SWidget> FVariantManagerVariantSetNode::GetCustomOutlinerContent(TSha
 
 						return FReply::Handled();
 					})
-					.ContentPadding(FMargin(2.0f, 1.0f))
-					.Content()
+					.ContentPadding(FMargin(1, 0))
+					.ToolTipText(LOCTEXT("AddVariantToolTip", "Add a new variant"))
 					[
-						SNew(STextBlock)
-						.TextStyle(FAppStyle::Get(), "NormalText.Important")
-						.Font(FAppStyle::Get().GetFontStyle("FontAwesome.10"))
-						.Text(FEditorFontGlyphs::Plus)
+						SNew(SImage)
+						.Image(FAppStyle::Get().GetBrush("Icons.PlusCircle"))
+						.ColorAndOpacity(FSlateColor::UseForeground())
 					]
 				]
 			]
@@ -728,22 +736,6 @@ UVariantSet& FVariantManagerVariantSetNode::GetVariantSet() const
 	return VariantSet;
 }
 
-FSlateColor FVariantManagerVariantSetNode::GetNodeBackgroundTint() const
-{
-	if (IsSelected())
-	{
-		return FAppStyle::GetSlateColor("SelectionColor_Pressed");
-	}
-	else if (IsHovered())
-	{
-		return FLinearColor(FColor(52, 52, 52, 255));
-	}
-	else
-	{
-		return FLinearColor(FColor(48, 48, 48, 255));
-	}
-}
-
 const FSlateBrush* FVariantManagerVariantSetNode::GetNodeBorderImage() const
 {
 	return IsExpanded() ? ExpandedBackgroundBrush : CollapsedBackgroundBrush;
diff --git a/Engine/Plugins/Enterprise/VariantManager/Source/VariantManager/Private/DisplayNodes/VariantManagerVariantSetNode.h b/Engine/Plugins/Enterprise/VariantManager/Source/VariantManager/Private/DisplayNodes/VariantManagerVariantSetNode.h
--- a/Engine/Plugins/Enterprise/VariantManager/Source/VariantManager/Private/DisplayNodes/VariantManagerVariantSetNode.h
+++ b/Engine/Plugins/Enterprise/VariantManager/Source/VariantManager/Private/DisplayNodes/VariantManagerVariantSetNode.h
@@ -22,6 +22,7 @@ public:
 	FVariantManagerVariantSetNode( UVariantSet& InVariantSet, TSharedPtr<FVariantManagerDisplayNode> InParentNode, TWeakPtr<FVariantManagerNodeTree> InParentTree );
 
 	// FVariantManagerDisplayNode interface
+	virtual const FTableRowStyle* GetRowStyle() const override;
 	virtual EVariantManagerNodeType GetType() const override;
 	virtual bool IsReadOnly() const override;
 	virtual FText GetDisplayName() const override;
@@ -32,7 +33,6 @@ public:
 	virtual TOptional<EItemDropZone> CanDrop(const FDragDropEvent& DragDropEvent, EItemDropZone ItemDropZone) const override;
 	virtual void Drop(const FDragDropEvent& DragDropEvent, EItemDropZone ItemDropZone) override;
 	virtual void BuildContextMenu(FMenuBuilder& MenuBuilder) override;
-	virtual FSlateColor GetNodeBackgroundTint() const override;
 	virtual const FSlateBrush* GetNodeBorderImage() const override;
 	virtual TSharedRef<SWidget> GetCustomOutlinerContent(TSharedPtr<SVariantManagerTableRow> InTableRow) override;
 	TSharedRef<SWidget> GetThumbnailWidget();
@@ -46,6 +46,8 @@ private:
 	/** The movie scene folder data which this node represents. */
 	UVariantSet& VariantSet;
 
+	FTableRowStyle RowStyle;
+
 	/** Different brushes so that the edges look good when expanded and collapsed */
 	const FSlateBrush* ExpandedBackgroundBrush;
 	const FSlateBrush* CollapsedBackgroundBrush;
diff --git a/Engine/Plugins/Enterprise/VariantManager/Source/VariantManager/Private/SDependencyRow.cpp b/Engine/Plugins/Enterprise/VariantManager/Source/VariantManager/Private/SDependencyRow.cpp
--- a/Engine/Plugins/Enterprise/VariantManager/Source/VariantManager/Private/SDependencyRow.cpp
+++ b/Engine/Plugins/Enterprise/VariantManager/Source/VariantManager/Private/SDependencyRow.cpp
@@ -9,6 +9,7 @@
 #include "Widgets/Images/SImage.h"
 #include "Widgets/Input/SButton.h"
 #include "Widgets/Input/SComboBox.h"
+#include "Widgets/Layout/SSeparator.h"
 
 #include "LevelVariantSets.h"
 #include "SVariantManager.h"
@@ -19,12 +20,16 @@
 
 #define LOCTEXT_NAMESPACE "SDependencyRow"
 
-void SDependencyRow::Construct(const FArguments& InArgs, const TSharedRef<STableViewBase>& InOwnerTableView, FColumnSizeData& InDependenciesColumnData, FVariantDependencyModelPtr InDependencyModel, bool bInteractionEnabled)
-{
-	STableRow<FVariantDependencyModelPtr>::ConstructInternal(STableRow::FArguments(), InOwnerTableView);
+const FName SDependencyRow::VisibilityColumn("Visibility");
+const FName SDependencyRow::VariantSetColumn("VariantSet");
+const FName SDependencyRow::VariantColumn("Variant");
 
+void SDependencyRow::Construct(const FArguments& InArgs, const TSharedRef<STableViewBase>& InOwnerTableView, FVariantDependencyModelPtr InDependencyModel, bool bInteractionEnabled)
+{
 	if (InDependencyModel.IsValid())
 	{
+		bIsDivider = InDependencyModel->bIsDivider;
+		bIsDependent = InDependencyModel->bIsDependent;
 		ParentVariantPtr = InDependencyModel->ParentVariant;
 		Dependency = InDependencyModel->Dependency;
 	}
@@ -37,138 +42,138 @@ void SDependencyRow::Construct(const FArguments& InArgs, const TSharedRef<STable
 	RebuildVariantSetOptions();
 	RebuildVariantOptions();
 
-	ChildSlot
-	[
-		SNew(SBox)
-		.VAlign(VAlign_Center)
-		.HAlign(HAlign_Fill)
-		.HeightOverride(26)
-		[
-			SNew(SSplitter)
-			.Orientation(Orient_Horizontal)
-
-			+ SSplitter::Slot()
-			.Value( InDependenciesColumnData.LeftColumnWidth )
-			.OnSlotResized(SSplitter::FOnSlotResized::CreateLambda([](float InNewWidth)
-			{
-				//This has to be bound or the splitter will take it upon itself to determine the size
-				//We do nothing here because it is handled by the column size data
-			}))
-			[
-				SNew(SBox)
+	SMultiColumnTableRow<FVariantDependencyModelPtr>::Construct(FSuperRowType::FArguments(), InOwnerTableView);
+}
+
+TSharedRef<SWidget> SDependencyRow::GenerateWidgetForColumn(const FName& InColumnName)
+{
+	if (bIsDivider)
+	{
+		if (InColumnName == VariantSetColumn || InColumnName == VariantColumn)
+		{
+			return SNew(SBox)
+				.Padding(0, 15, 0, 15)
 				.VAlign(VAlign_Center)
-				.HAlign(HAlign_Fill)
-				.HeightOverride(21)
 				[
-					SNew(SComboBox<TSharedPtr<FText>>)
-					.OptionsSource(&VariantSetOptions)
-					.IsEnabled(bInteractionEnabled)
-					.OnGenerateWidget_Lambda([](TSharedPtr<FText> Item)
-					{
-						return SNew(STextBlock).Text(*Item);
-					})
-					.Content()
-					[
-						SNew(STextBlock)
-						.Text(this, bInteractionEnabled ? &SDependencyRow::GetSelectedVariantSetOption : &SDependencyRow::GetDependentVariantSetText)
+					SNew(SHorizontalBox)
+					+ SHorizontalBox::Slot()
+					.FillWidth(1)
+					[	
+						SNew(SSeparator)
+						.Orientation(EOrientation::Orient_Horizontal)
+						.Thickness(1)
+						.SeparatorImage(FAppStyle::GetBrush("WhiteBrush"))
 					]
-					.OnSelectionChanged(this, &SDependencyRow::OnSelectedVariantSetChanged)
-				]
-			]
-
-			+ SSplitter::Slot()
-			.Value( InDependenciesColumnData.MiddleColumnWidth )
-			.OnSlotResized( InDependenciesColumnData.OnFirstSplitterChanged )
-			[
-				SNew(SBox)
-				.VAlign(VAlign_Center)
-				.HAlign(HAlign_Fill)
+				];
+		}
+	}
+	else if (ParentVariantPtr.IsValid())
+	{
+		if (InColumnName == VisibilityColumn)
+		{
+			return SNew(SBox)
 				.HeightOverride(21)
+				.WidthOverride(21)
 				[
-					SNew(SComboBox<TSharedPtr<FText>>)
-					.OptionsSource(&VariantOptions)
-					.IsEnabled(bInteractionEnabled)
-					.OnGenerateWidget_Lambda([](TSharedPtr<FText> Item)
-					{
-						return SNew(STextBlock).Text(*Item);
-					})
-					.Content()
+					SNew(SButton)
+					.IsFocusable(false)
+					.ToolTipText(LOCTEXT("ToggleDependency", "Enable or disable this dependency"))
+					.ButtonStyle(FAppStyle::Get(), "HoverHintOnly")
+					.ContentPadding(0.0f)
+					.HAlign(HAlign_Center)
+					.VAlign(VAlign_Center)
+					.Visibility_Raw(this, &SDependencyRow::OnGetEyeIconVisibility)
+					.OnClicked(this, &SDependencyRow::OnEnableRowToggled)
 					[
-						SNew(STextBlock)
-						.Text(this, bInteractionEnabled ? &SDependencyRow::GetSelectedVariantOption : &SDependencyRow::GetDependentVariantText)
+						SNew(SImage)
+						.Image_Lambda([this]()
+						{
+							if (Dependency && Dependency->bEnabled)
+							{
+								return FAppStyle::GetBrush("Level.VisibleIcon16x");
+							}
+							return FAppStyle::GetBrush("Level.NotVisibleIcon16x");
+						})
 					]
-					.OnSelectionChanged(this, &SDependencyRow::OnSelectedVariantChanged)
-				]
-			]
-
-			+ SSplitter::Slot()
-			.Value( InDependenciesColumnData.RightColumnWidth )
-			.OnSlotResized( InDependenciesColumnData.OnSecondSplitterChanged )
-			[
-				SNew(SHorizontalBox)
-				+ SHorizontalBox::Slot()
+				];
+		}
+		else if (InColumnName == VariantSetColumn)
+		{
+			return SNew(SBox)
 				.VAlign(VAlign_Center)
 				.HAlign(HAlign_Fill)
-				.Padding(FMargin(1.0f, 0.0f, 3.0f, 0.0))
-				.AutoWidth()
+				.HeightOverride(21)
 				[
-					SNew(SBox)
-					.HeightOverride(21)
-					.WidthOverride(21)
+					SNew(SHorizontalBox)
+					+ SHorizontalBox::Slot()
+					.FillWidth(0.7f)
 					[
-						SNew(SButton)
-						.IsFocusable(false)
-						.ToolTipText(LOCTEXT("DeleteDependency", "Delete this dependency"))
-						.ButtonStyle( FAppStyle::Get(), "HoverHintOnly" )
-						.ContentPadding(0.0f)
-						.HAlign(HAlign_Center)
-						.VAlign(VAlign_Center)
-						.Visibility(bInteractionEnabled ? EVisibility::Visible : EVisibility::Hidden)
-						.OnClicked(this, &SDependencyRow::OnDeleteRowClicked)
+						SNew(SComboBox<TSharedPtr<FText>>)
+						.ComboBoxStyle(&FAppStyle::Get().GetWidgetStyle<FComboBoxStyle>(TEXT("SimpleComboBox")))
+						.OptionsSource(&VariantSetOptions)
+						.OnGenerateWidget_Lambda([](TSharedPtr<FText> Item)
+						{
+							return SNew(STextBlock).Text(*Item);
+						})
+						.Content()
 						[
 							SNew(STextBlock)
-							.TextStyle(FAppStyle::Get(), "NormalText.Important")
-							.Font(FAppStyle::Get().GetFontStyle("FontAwesome.10"))
-							.Text(FEditorFontGlyphs::Trash)
+							.Text(this, bIsDependent ? &SDependencyRow::GetDependentVariantSetText : &SDependencyRow::GetSelectedVariantSetOption)
 						]
+						.OnSelectionChanged(this, &SDependencyRow::OnSelectedVariantSetChanged)
 					]
-				]
 
-				+ SHorizontalBox::Slot()
-				.Padding(FMargin(0.0f, 0.0f, 3.0f, 0.0))
+					+ SHorizontalBox::Slot()
+					.FillWidth(0.3f)
+					[
+						SNew(SSpacer)
+					]
+				];
+		}
+		else
+		{ 
+			return SNew(SBox)
 				.VAlign(VAlign_Center)
 				.HAlign(HAlign_Fill)
-				.AutoWidth()
+				.HeightOverride(21)
 				[
-					SNew(SBox)
-					.HeightOverride(21)
-					.WidthOverride(21)
+					SNew(SHorizontalBox)
+					+ SHorizontalBox::Slot()
+					.FillWidth(0.7f)
 					[
-						SNew(SButton)
-						.IsFocusable(false)
-						.ToolTipText(LOCTEXT("ToggleDependency", "Enable or disable this dependency"))
-						.ButtonStyle( FAppStyle::Get(), "HoverHintOnly" )
-						.ContentPadding(0.0f)
-						.HAlign(HAlign_Center)
-						.VAlign(VAlign_Center)
-						.Visibility(bInteractionEnabled ? EVisibility::Visible : EVisibility::Hidden)
-						.OnClicked(this, &SDependencyRow::OnEnableRowToggled)
+						SNew(SComboBox<TSharedPtr<FText>>)
+						.ComboBoxStyle(&FAppStyle::Get().GetWidgetStyle<FComboBoxStyle>(TEXT("SimpleComboBox")))
+						.OptionsSource(&VariantOptions)
+						.OnGenerateWidget_Lambda([](TSharedPtr<FText> Item)
+						{
+							return SNew(STextBlock).Text(*Item);
+						})
+						.Content()
 						[
-							SNew(SImage)
-							.Image_Lambda([this]()
-							{
-								if(Dependency && Dependency->bEnabled)
-								{
-									return FAppStyle::GetBrush("Level.VisibleIcon16x");
-								}
-								return FAppStyle::GetBrush("Level.NotVisibleIcon16x");
-							})
+							SNew(STextBlock)
+							.Text(this, bIsDependent ? &SDependencyRow::GetDependentVariantText  : &SDependencyRow::GetSelectedVariantOption)
 						]
+						.OnSelectionChanged(this, &SDependencyRow::OnSelectedVariantChanged)
+					]
+					+ SHorizontalBox::Slot()
+					.FillWidth(0.3f)
+					[
+						SNew(SSpacer)
 					]
-				]
-			]
-		]
-	];
+				];
+		}
+	}
+
+	return SNew(SBox);
+}
+
+EVisibility SDependencyRow::OnGetEyeIconVisibility() const
+{
+	if (bIsDependent)
+	{
+		return EVisibility::Collapsed;
+	}
+	return (Dependency && !Dependency->bEnabled) || IsHovered() || IsSelected() ? EVisibility::Visible : EVisibility::Hidden;
 }
 
 void SDependencyRow::OnSelectedVariantSetChanged(TSharedPtr<FText> NewItem, ESelectInfo::Type SelectType)
@@ -382,35 +387,6 @@ void SDependencyRow::RebuildVariantOptions()
 	}
 }
 
-FReply SDependencyRow::OnDeleteRowClicked()
-{
-	UVariant* ParentVariant = ParentVariantPtr.Get();
-	if (Dependency && ParentVariant)
-	{
-		int32 DependencyIndex = INDEX_NONE;
-		for (int32 Index = 0; Index < ParentVariant->GetNumDependencies(); ++Index)
-		{
-			if (&ParentVariant->GetDependency(Index) == Dependency)
-			{
-				DependencyIndex = Index;
-				break;
-			}
-		}
-
-		if (DependencyIndex != INDEX_NONE)
-		{
-			FScopedTransaction Transaction(FText::Format(
-				LOCTEXT("DeleteDependencyTransaction", "Delete a dependency from variant '{0}'"),
-				ParentVariant->GetDisplayText()
-			));
-
-			ParentVariant->DeleteDependency(DependencyIndex);
-		}
-	}
-
-	return FReply::Handled();
-}
-
 FReply SDependencyRow::OnEnableRowToggled()
 {
 	UVariant* ParentVariant = ParentVariantPtr.Get();
diff --git a/Engine/Plugins/Enterprise/VariantManager/Source/VariantManager/Private/SDependencyRow.h b/Engine/Plugins/Enterprise/VariantManager/Source/VariantManager/Private/SDependencyRow.h
--- a/Engine/Plugins/Enterprise/VariantManager/Source/VariantManager/Private/SDependencyRow.h
+++ b/Engine/Plugins/Enterprise/VariantManager/Source/VariantManager/Private/SDependencyRow.h
@@ -14,18 +14,38 @@ struct FVariantDependency;
 // Adapter so that we can use arrays of these objects on SListViews and still reference the original Dependency
 struct FVariantDependencyModel : TSharedFromThis<FVariantDependencyModel>
 {
+	FVariantDependencyModel(
+		TWeakObjectPtr<UVariant> InParentVariant,
+		FVariantDependency* InDependency,
+		bool bInIsDependent,
+		bool bInIsDivider)
+		: ParentVariant(InParentVariant)
+		, Dependency(InDependency)
+		, bIsDependent(bInIsDependent)
+		, bIsDivider(bInIsDivider)
+	{}
+
 	TWeakObjectPtr<UVariant> ParentVariant;
-	FVariantDependency* Dependency;
+	FVariantDependency* Dependency = nullptr;
+	bool bIsDependent = false;
+	bool bIsDivider = false;
 };
+
 using FVariantDependencyModelPtr = TSharedPtr<FVariantDependencyModel>;
 
-class SDependencyRow : public STableRow<FVariantDependencyModelPtr>
+class SDependencyRow : public SMultiColumnTableRow<FVariantDependencyModelPtr>
 {
 public:
 	SLATE_BEGIN_ARGS(SDependencyRow) {}
 	SLATE_END_ARGS()
 
-	void Construct(const FArguments& InArgs, const TSharedRef<STableViewBase>& InOwnerTableView, FColumnSizeData& InDependenciesColumnData, FVariantDependencyModelPtr InDependencyModel, bool bInteractionEnabled);
+	static const FName VisibilityColumn;
+	static const FName VariantSetColumn;
+	static const FName VariantColumn;
+
+	void Construct(const FArguments& InArgs, const TSharedRef<STableViewBase>& InOwnerTableView, FVariantDependencyModelPtr InDependencyModel, bool bInteractionEnabled);
+	
+	virtual TSharedRef<SWidget> GenerateWidgetForColumn(const FName& InColumnName) override;
 
 private:
 	void OnSelectedVariantSetChanged(TSharedPtr<FText> NewItem, ESelectInfo::Type SelectType);
@@ -40,13 +60,16 @@ private:
 	void RebuildVariantSetOptions();
 	void RebuildVariantOptions();
 
-	FReply OnDeleteRowClicked();
 	FReply OnEnableRowToggled();
+	EVisibility OnGetEyeIconVisibility() const;
 
 private:
 	TArray<TSharedPtr<FText>> VariantSetOptions;
 	TArray<TSharedPtr<FText>> VariantOptions;
 
 	TWeakObjectPtr<UVariant> ParentVariantPtr;
-	FVariantDependency* Dependency;
+	FVariantDependency* Dependency = nullptr;
+
+	bool bIsDivider = false;
+	bool bIsDependent = false;
 };
\ No newline at end of file
diff --git a/Engine/Plugins/Enterprise/VariantManager/Source/VariantManager/Private/SVariantManager.cpp b/Engine/Plugins/Enterprise/VariantManager/Source/VariantManager/Private/SVariantManager.cpp
--- a/Engine/Plugins/Enterprise/VariantManager/Source/VariantManager/Private/SVariantManager.cpp
+++ b/Engine/Plugins/Enterprise/VariantManager/Source/VariantManager/Private/SVariantManager.cpp
@@ -64,6 +64,7 @@
 #include "SceneOutlinerPublicTypes.h"
 #include "ActorTreeItem.h"
 #include "ScopedTransaction.h"
+#include "SPositiveActionButton.h"
 #include "Widgets/Input/SButton.h"
 #include "Widgets/Input/SCheckBox.h"
 #include "Widgets/Input/SSearchBox.h"
@@ -73,6 +74,7 @@
 #include "Widgets/Layout/SSplitter.h"
 #include "Widgets/Notifications/SNotificationList.h"
 #include "Widgets/Text/STextBlock.h"
+#include "Widgets/Text/SRichTextBlock.h"
 #include "Widgets/Views/SListView.h"
 #include "Widgets/Views/STableRow.h"
 #include "Widgets/Views/STableViewBase.h"
@@ -93,48 +95,23 @@ FSplitterValues::FSplitterValues(FString& InSerialized)
 	}
 
 	VariantColumn = FCString::Atof(*SplitString[0]);
-	ActorColumn = FCString::Atof(*SplitString[1]);
-	PropertyNameColumn = FCString::Atof(*SplitString[2]);
-	DependenciesVariantSetsColumn = FCString::Atof(*SplitString[3]);
-	DependenciesVariantColumn = FCString::Atof(*SplitString[4]);
+	PropertyNameColumn = FCString::Atof(*SplitString[1]);
+	PropertyValueColumn = FCString::Atof(*SplitString[2]);
 }
 
 FString FSplitterValues::ToString()
 {
 	return FString::SanitizeFloat(VariantColumn) + TEXT(";") +
-		   FString::SanitizeFloat(ActorColumn) + TEXT(";") +
 		   FString::SanitizeFloat(PropertyNameColumn) + TEXT(";") +
-		   FString::SanitizeFloat(DependenciesVariantSetsColumn) + TEXT(";") +
-		   FString::SanitizeFloat(DependenciesVariantColumn);
+		   FString::SanitizeFloat(PropertyValueColumn);
 }
 
 TSharedRef<SWidget> SVariantManager::MakeAddButton()
 {
-	return SNew(SButton)
-	.OnClicked(this, &SVariantManager::OnAddVariantSetClicked)
-	.ButtonStyle(FAppStyle::Get(), "FlatButton.Success")
-	.Content()
-	[
-		SNew(SHorizontalBox)
-
-		+ SHorizontalBox::Slot()
-		.VAlign(VAlign_Center)
-		.AutoWidth()
-		[
-			SNew(STextBlock)
-			.TextStyle(FAppStyle::Get(), "NormalText.Important")
-			.Font(FAppStyle::Get().GetFontStyle("FontAwesome.10"))
-			.Text(FEditorFontGlyphs::Plus)
-		]
-		+ SHorizontalBox::Slot()
-		.Padding(4, 0, 0, 0)
-		.AutoWidth()
-		[
-			SNew(STextBlock)
-			.TextStyle(FAppStyle::Get(), "NormalText.Important")
-			.Text(LOCTEXT("VariantSetPlusText", "Variant Set"))
-		]
-	];
+	return SNew(SPositiveActionButton)
+		.Icon(FAppStyle::Get().GetBrush("Icons.Plus"))
+		.Text(LOCTEXT("VariantSetPlusText", "VARIANT SET"))
+		.OnClicked(this, &SVariantManager::OnAddVariantSetClicked);
 }
 
 TSharedRef<ITableRow> SVariantManager::MakeCapturedPropertyRow(TSharedPtr<FVariantManagerPropertyNode> Item, const TSharedRef<STableViewBase>& OwnerTable)
@@ -171,6 +148,9 @@ void SVariantManager::Construct(const FArguments& InArgs, TSharedRef<FVariantMan
 
 	bAutoCaptureProperties = false;
 
+	HeaderStyle = &FAppStyle::Get().GetWidgetStyle<FTableColumnHeaderStyle>("TableView.Header.Column");
+	SplitterStyle = &FAppStyle::Get().GetWidgetStyle<FSplitterStyle>("Splitter");
+	
 	CreateCommandBindings();
 
 	SAssignNew(NodeTreeView, SVariantManagerNodeTreeView, InVariantManager->GetNodeTree());
@@ -185,28 +165,18 @@ void SVariantManager::Construct(const FArguments& InArgs, TSharedRef<FVariantMan
 	}
 
 	PropertiesNameColumnWidth = SplitterValues.PropertyNameColumn;
-	PropertiesValueColumnWidth = 1.0f - SplitterValues.ActorColumn - SplitterValues.PropertyNameColumn;
-	DependenciesVariantColumnWidth = SplitterValues.DependenciesVariantColumn;
-	DependenciesControlColumnWidth = 1.0f - SplitterValues.DependenciesVariantSetsColumn - SplitterValues.DependenciesVariantColumn;
-
-	PropertiesColumnSizeData.LeftColumnWidth = TAttribute<float>( this, &SVariantManager::OnGetPropertiesActorColumnWidth );
-	PropertiesColumnSizeData.MiddleColumnWidth = TAttribute<float>( this, &SVariantManager::OnGetPropertiesNameColumnWidth );
-	PropertiesColumnSizeData.RightColumnWidth = TAttribute<float>( this, &SVariantManager::OnGetPropertiesValueColumnWidth );
-	PropertiesColumnSizeData.OnFirstSplitterChanged = SSplitter::FOnSlotResized::CreateSP( this, &SVariantManager::OnSetPropertiesNameColumnWidth );
-	PropertiesColumnSizeData.OnSecondSplitterChanged = SSplitter::FOnSlotResized::CreateLambda( [ this ]( float InNewWidth )
-	{
-		// We want the actor column to remain fixed while we move the splitter between property names and values
-		float PropertyComboWidth = PropertiesNameColumnWidth + PropertiesValueColumnWidth;
-		float OldActorColumnWidth = 1.0f - PropertyComboWidth;
-		PropertiesValueColumnWidth = InNewWidth;
-		PropertiesNameColumnWidth = 1.0f - PropertiesValueColumnWidth - OldActorColumnWidth;
-	} );
+	PropertiesValueColumnWidth = 1.0f - SplitterValues.PropertyNameColumn;
 
-	DependenciesColumnSizeData.LeftColumnWidth = TAttribute<float>( this, &SVariantManager::OnGetDependenciesVariantSetColumnWidth );
-	DependenciesColumnSizeData.MiddleColumnWidth = TAttribute<float>( this, &SVariantManager::OnGetDependenciesVariantColumnWidth );
-	DependenciesColumnSizeData.RightColumnWidth = TAttribute<float>( this, &SVariantManager::OnGetDependenciesControlColumnWidth );
-	DependenciesColumnSizeData.OnFirstSplitterChanged = SSplitter::FOnSlotResized::CreateSP( this, &SVariantManager::OnSetDependenciesVariantColumnWidth );
-	DependenciesColumnSizeData.OnSecondSplitterChanged = SSplitter::FOnSlotResized::CreateSP( this, &SVariantManager::OnSetDependenciesControlColumnWidth );
+	PropertiesColumnSizeData.NameColumnWidth = TAttribute<float>( this, &SVariantManager::OnGetPropertiesNameColumnWidth );
+	PropertiesColumnSizeData.ValueColumnWidth = TAttribute<float>( this, &SVariantManager::OnGetPropertiesValueColumnWidth );
+	PropertiesColumnSizeData.OnSplitterNameColumnChanged = SSplitter::FOnSlotResized::CreateLambda( [ this ]( float InNewWidth )
+	{
+		PropertiesNameColumnWidth = InNewWidth;
+	});
+	PropertiesColumnSizeData.OnSplitterValueColumnChanged = SSplitter::FOnSlotResized::CreateLambda([this](float InNewWidth)
+	{
+		PropertiesValueColumnWidth = InNewWidth;
+	});
 
 	InVariantManager->GetSelection().GetOnOutlinerNodeSelectionChanged().AddSP(this, &SVariantManager::OnOutlinerNodeSelectionChanged);
 	InVariantManager->GetSelection().GetOnActorNodeSelectionChanged().AddSP(this, &SVariantManager::OnActorNodeSelectionChanged);
@@ -244,8 +214,6 @@ void SVariantManager::Construct(const FArguments& InArgs, TSharedRef<FVariantMan
 	}
 
 	RightTreeRootItems.Empty();
-	RightTreeRootItems.Reserve(2);
-	RightTreeRootItems.Add(MakeShared<ERightTreeRowType>(ERightTreeRowType::PropertiesHeader));
 	RightTreeRootItems.Add(MakeShared<ERightTreeRowType>(ERightTreeRowType::DependenciesHeader));
 
 	TSharedPtr<SHeaderRow> CollapsedHeader = SNew( SHeaderRow ).Visibility( EVisibility::Collapsed );
@@ -284,7 +252,7 @@ void SVariantManager::Construct(const FArguments& InArgs, TSharedRef<FVariantMan
 		RightTree->SetItemExpansion(Item, true);
 	}
 
-	float BorderThickness = FVariantManagerStyle::Get()->GetFloat( "VariantManager.Spacings.BorderThickness" );
+	const float BorderThickness = FVariantManagerStyle::Get().GetFloat("VariantManager.Spacings.BorderThickness");;
 
 	ChildSlot
 	[
@@ -311,51 +279,73 @@ void SVariantManager::Construct(const FArguments& InArgs, TSharedRef<FVariantMan
 				]
 
 				+SHorizontalBox::Slot()
+				.FillWidth(1.0f)
+				.HAlign(HAlign_Left)
 				.VAlign(VAlign_Center)
-				.Padding(FMargin(0.f, 0.f, VM_COMMON_PADDING+2.0f, 0.f))
-				.AutoWidth()
-				.MaxWidth(VM_COMMON_HEADER_MAX_HEIGHT) // square aspect ratio
 				[
-					SNew(SCheckBox)
-					.Style( FVariantManagerStyle::Get(), "AutoCaptureCheckbox" )
-					.Padding( FMargin( 10.f, 10.f, 10.0f, 10.f ) ) // Give some space to show up the style image or else it will collapse to nothing since it has no content
-					.ToolTipText(LOCTEXT("AutoCaptureTooltip", "Enable or disable auto-capturing properties"))
-					.IsChecked_Lambda([&bAutoCaptureProperties = bAutoCaptureProperties]()
-					{
-						return bAutoCaptureProperties? ECheckBoxState::Checked : ECheckBoxState::Unchecked;
-					})
-					.OnCheckStateChanged_Lambda([&bAutoCaptureProperties = bAutoCaptureProperties](const ECheckBoxState NewState)
-					{
-						bAutoCaptureProperties = NewState == ECheckBoxState::Checked;
-					})
+					SNew(SSearchBox)
+					.HintText(LOCTEXT("VariantManagerFilterText", "Filter variants"))
+					.OnTextChanged(this, &SVariantManager::OnOutlinerSearchChanged)
+					.MinDesiredWidth(200)
 				]
 
 				+SHorizontalBox::Slot()
-				.FillWidth(1.0f)
 				.VAlign(VAlign_Center)
-				.MaxWidth( 200.0f )
+				.HAlign(HAlign_Right)
+				.AutoWidth()
+				.Padding(FMargin(0.f, 0.f, VM_COMMON_PADDING+2.0f, 0.f))
 				[
-					SNew(SSearchBox)
-					.HintText(LOCTEXT("VariantManagerFilterText", "Filter variants"))
-					.OnTextChanged(this, &SVariantManager::OnOutlinerSearchChanged)
+					SNew(SHorizontalBox)
+					+ SHorizontalBox::Slot()
+					.MaxWidth(VM_COMMON_HEADER_MAX_HEIGHT) // square aspect ratio
+					.AutoWidth()
+					.VAlign(VAlign_Center)
+					[
+						SNew(SCheckBox)
+						.Style( FVariantManagerStyle::Get(), "VariantManager.AutoCapture" )
+						.Padding( FMargin( 10.f, 10.f, 10.0f, 10.f ) ) // Give some space to show up the style image or else it will collapse to nothing since it has no content
+						.ToolTipText(LOCTEXT("AutoCaptureTooltip", "Enable or disable auto-capturing properties"))
+						.IsChecked_Lambda([&bAutoCaptureProperties = bAutoCaptureProperties]()
+						{
+							return bAutoCaptureProperties? ECheckBoxState::Checked : ECheckBoxState::Unchecked;
+						})
+						.OnCheckStateChanged_Lambda([&bAutoCaptureProperties = bAutoCaptureProperties](const ECheckBoxState NewState)
+						{
+							bAutoCaptureProperties = NewState == ECheckBoxState::Checked;
+						})
+					]
+					
+					+ SHorizontalBox::Slot()
+					.AutoWidth()
+					.VAlign(VAlign_Center)
+					.Padding(5, 0, 3, 0)
+					[
+						SNew(STextBlock)
+						.Text(LOCTEXT("VariantManagerAutoCapture", "Auto Capture"))
+					]
 				]
 			]
 		]
 
+		+ SVerticalBox::Slot()
+		.AutoHeight()
+		[
+			SNew(SSeparator)
+			.Orientation(EOrientation::Orient_Horizontal)
+			.Thickness(3)
+			.SeparatorImage(&SplitterStyle->HandleNormalBrush)
+		]
+		
 		+ SVerticalBox::Slot()
 		.VAlign( VAlign_Fill )
 		.FillHeight( 1.0f )
 		[
-			// Common "background"
-			SNew(SBorder)
-			.BorderImage( FAppStyle::GetBrush( "WhiteBrush" ) )
-			.BorderBackgroundColor( FVariantManagerStyle::Get()->GetColor("VariantManager.Panels.LightBackgroundColor") )
-			.Padding( FMargin( BorderThickness ) )
+			SNew(SBox)
 			[
 				SAssignNew(MainSplitter, SSplitter)
 				.Orientation(Orient_Horizontal)
-				.PhysicalSplitterHandleSize( BorderThickness )
-				.HitDetectionSplitterHandleSize( BorderThickness )
+				.PhysicalSplitterHandleSize(BorderThickness)
+				.HitDetectionSplitterHandleSize(BorderThickness)
 
 				// VariantSet/Variant column
 				+SSplitter::Slot()
@@ -374,12 +364,19 @@ void SVariantManager::Construct(const FArguments& InArgs, TSharedRef<FVariantMan
 				+ SSplitter::Slot()
 				.Value( 1.0f - SplitterValues.VariantColumn )
 				[
-					SNew( SScrollBorder, RightTree )
+					SNew(SSplitter)
+					.Orientation(Orient_Vertical)
+					.PhysicalSplitterHandleSize(BorderThickness)
+					.HitDetectionSplitterHandleSize(BorderThickness)
+					
+					+ SSplitter::Slot()
 					[
-						SNew( SBox ) // Very important to prevent the tree from expanding freely
-						[
-							RightTree
-						]
+						GenerateRightTreePropertiesRowContent()
+					]
+					
+					+ SSplitter::Slot()
+					[
+						RightTree
 					]
 				]
 			]
@@ -437,10 +434,8 @@ SVariantManager::~SVariantManager()
 		{
 			FSplitterValues Values;
 			Values.VariantColumn = MainSplitter->SlotAt(0).GetSizeValue();
-			Values.ActorColumn = OnGetPropertiesActorColumnWidth();;
 			Values.PropertyNameColumn = OnGetPropertiesNameColumnWidth();
-			Values.DependenciesVariantSetsColumn = OnGetDependenciesVariantSetColumnWidth();
-			Values.DependenciesVariantColumn = OnGetDependenciesVariantColumnWidth();
+			Values.PropertyValueColumn = OnGetPropertiesValueColumnWidth();
 
 			GConfig->SetString(TEXT("VariantManager"), TEXT("SplitterValues"), *Values.ToString(), GEditorPerProjectIni);
 		}
@@ -2118,7 +2113,6 @@ void SVariantManager::RefreshActorList()
 			if (Bindings.Num() > 0)
 			{
 				TargetBindings.Append(Bindings);
-				TargetBindings.Add(nullptr); // nullptrs will be converted to spacers
 			}
 		}
 
@@ -2353,7 +2347,7 @@ void SVariantManager::RefreshPropertyList()
 
 void SVariantManager::RefreshDependencyLists()
 {
-	if (!DependenciesList || !DependentsList)
+	if (!DependenciesList)
 	{
 		return;
 	}
@@ -2365,7 +2359,6 @@ void SVariantManager::RefreshDependencyLists()
 	}
 
 	DisplayedDependencies.Reset();
-	DisplayedDependents.Reset();
 
 	TArray<UVariant*> SelectedVariants;
 	TArray<UVariantSet*> SelectedVariantSets;
@@ -2376,26 +2369,37 @@ void SVariantManager::RefreshDependencyLists()
 
 		for (int32 Index = 0; Index < SelectedVariant->GetNumDependencies(); ++Index)
 		{
-			FVariantDependencyModelPtr Model = MakeShared<FVariantDependencyModel>();
-			Model->Dependency = &SelectedVariant->GetDependency(Index);
-			Model->ParentVariant = SelectedVariant;
+			FVariantDependencyModelPtr Model = MakeShared<FVariantDependencyModel>(
+				SelectedVariant, 
+				&SelectedVariant->GetDependency(Index), 
+				false,
+				false);
 
 			DisplayedDependencies.Add(Model);
 		}
 
 		const bool bOnlyEnabledDependencies = false;
 		TArray<UVariant*> Dependents = SelectedVariant->GetDependents(PinnedVariantManager->GetCurrentLevelVariantSets(), bOnlyEnabledDependencies);
+		
+		if (Dependents.Num() > 0)
+		{
+			// Add divider between the dependencies list and dependents list
+			DisplayedDependencies.Add(MakeShared<FVariantDependencyModel>(nullptr, nullptr, false, true));
+		}
+
 		for (UVariant* Dependent : Dependents)
 		{
-			FVariantDependencyModelPtr Model = MakeShared<FVariantDependencyModel>();
-			Model->ParentVariant = Dependent;
+			FVariantDependencyModelPtr Model = MakeShared<FVariantDependencyModel>(
+				Dependent,
+				nullptr,
+				true,
+				false);
 
-			DisplayedDependents.Add(Model);
+			DisplayedDependencies.Add(Model);
 		}
 	}
 
 	DependenciesList->RequestListRefresh();
-	DependentsList->RequestListRefresh();
 }
 
 void SVariantManager::UpdatePropertyDefaults()
@@ -3373,7 +3377,6 @@ TSharedRef<ITableRow> SVariantManager::GenerateRightTreeRow( TSharedRef<ERightTr
 	TSharedPtr<SWidget> RowContent = SNullWidget::NullWidget;
 	switch ( *RowType )
 	{
-	case ERightTreeRowType::PropertiesHeader:
 	case ERightTreeRowType::DependenciesHeader:
 		RowContent = GenerateRightTreeHeaderRowContent(*RowType, TableRow);
 		break;
@@ -3393,112 +3396,135 @@ TSharedRef<ITableRow> SVariantManager::GenerateRightTreeRow( TSharedRef<ERightTr
 
 TSharedRef<SWidget> SVariantManager::GenerateRightTreeHeaderRowContent( ERightTreeRowType RowType, TSharedRef<STableRow<TSharedRef<ERightTreeRowType>>> InTableRow )
 {
-	FText RowText;
-	switch ( RowType )
-	{
-	case ERightTreeRowType::DependenciesHeader:
-		RowText = FText::FromString(TEXT("Dependencies"));
-		break;
-	case ERightTreeRowType::PropertiesHeader:
-		RowText = FText::FromString(TEXT("Properties"));
-		break;
-	default:
-		return SNullWidget::NullWidget;
-	}
-
-	TSharedRef<SHorizontalBox> HorizontalBox = SNew( SHorizontalBox )
-	.IsEnabled_Lambda([this, RowType]()
-	{
-		if (RowType == ERightTreeRowType::PropertiesHeader)
-		{
-			return true;
-		}
+	return
+		SNew(SBox)
+		[
+			SNew(SBorder)
+			.VAlign( VAlign_Center )
+			.HAlign(HAlign_Fill)
+			.BorderImage(&HeaderStyle->NormalBrush)
+			.Padding(FMargin(10, 7, 8, 8))
+			[
+				SNew(SHorizontalBox)
+				.IsEnabled_Lambda([this]()
+				{
+					TSharedPtr<FVariantManager> PinnedVariantManager = VariantManagerPtr.Pin();
+					if (PinnedVariantManager.IsValid())
+					{
+						TArray<UVariant*> SelectedVariants;
+						TArray<UVariantSet*> SelectedVariantSets;
+						PinnedVariantManager->GetSelection().GetSelectedVariantsAndVariantSets(SelectedVariants, SelectedVariantSets);
 
-		TSharedPtr<FVariantManager> PinnedVariantManager = VariantManagerPtr.Pin();
-		if (PinnedVariantManager.IsValid())
-		{
-			TArray<UVariant*> SelectedVariants;
-			TArray<UVariantSet*> SelectedVariantSets;
-			PinnedVariantManager->GetSelection().GetSelectedVariantsAndVariantSets(SelectedVariants, SelectedVariantSets);
+						return SelectedVariants.Num() == 1;
+					}
+					return false;
+				})
+			
+				+ SHorizontalBox::Slot()
+				.Padding(FMargin(0.f, 0.f, 4.f, 0.f))
+				.VAlign(VAlign_Center)
+				.HAlign(HAlign_Left)
+				.AutoWidth()
+				[
+					SNew(SExpanderArrow, InTableRow).IndentAmount(FVariantManagerStyle::Get().GetFloat("VariantManager.Spacings.IndentAmount"))
+				]
 
-			return SelectedVariants.Num() == 1;
-		}
-		return false;
-	});
+				+ SHorizontalBox::Slot()
+				.VAlign(VAlign_Center)
+				.HAlign(HAlign_Left)
+				.FillWidth(1.0f)
+				[
+					SNew(SRichTextBlock)
+					.Text(LOCTEXT("Dependencies", "Dependencies"))
+					.TransformPolicy(ETextTransformPolicy::ToUpper)
+					.DecoratorStyleSet(&FAppStyle::Get())
+					.TextStyle(FAppStyle::Get(), "DetailsView.CategoryTextStyle")
+				]
 
-	HorizontalBox->AddSlot()
-	.Padding( FMargin( 4.f, 0.f, 4.f, 0.f ) )
-	.VAlign( VAlign_Center )
-	.MaxWidth( 24.0f )
-	[
-		SNew( SExpanderArrow, InTableRow ).IndentAmount( FVariantManagerStyle::Get()->GetFloat("VariantManager.Spacings.IndentAmount") )
-	];
+				+ SHorizontalBox::Slot()
+				.VAlign(VAlign_Center)
+				.HAlign(HAlign_Right)
+				[
+					SNew(SHorizontalBox)
+					+ SHorizontalBox::Slot()
+					[
+						SNew(SButton)
+						.ButtonStyle(FAppStyle::Get(), "HoverHintOnly")
+						.OnClicked(this, &SVariantManager::OnAddDependencyClicked)
+						.ContentPadding(FMargin(1, 0))
+						.ToolTipText(LOCTEXT("AddDependencyTooltip", "Add a new dependency to this variant"))
+						[
+							SNew(SImage)
+							.Image(FAppStyle::Get().GetBrush("Icons.PlusCircle"))
+							.ColorAndOpacity(FSlateColor::UseForeground())
+						]
+					]
 
-	HorizontalBox->AddSlot()
-		.VAlign( VAlign_Center )
-		.HAlign( HAlign_Left )
-		.AutoWidth()
-		.Padding( FMargin( 0.f, 0.f, 4.f, 0.f ) )
-		[
-			SNew( STextBlock )
-			.TextStyle( FAppStyle::Get(), "NormalText.Important" )
-			.Text( RowText )
+					+ SHorizontalBox::Slot()
+					.VAlign(VAlign_Center)
+					.MaxWidth(24.0f)
+					.Padding(FMargin(4.f, 0.f, 4.f, 0.f))
+					[
+						SNew(SBox)
+						[
+							SNew(SButton)
+							.ButtonStyle(FAppStyle::Get(), "HoverHintOnly")
+							.OnClicked(this, &SVariantManager::OnDeleteSelectedDependencies)
+							.ContentPadding(FMargin(1, 0))
+							.ToolTipText(LOCTEXT("DeleteDependency", "Delete selected dependencies"))
+							[
+								SNew(SImage)
+								.Image(FAppStyle::Get().GetBrush("Icons.Delete"))
+								.ColorAndOpacity(FSlateColor::UseForeground())
+							]
+						]
+					]
+				]
+			]
 		];
+}
 
-	if ( RowType == ERightTreeRowType::DependenciesHeader )
+FReply SVariantManager::OnDeleteSelectedDependencies()
+{
+	TArray<FVariantDependencyModelPtr> SelectedItems = DependenciesList->GetSelectedItems();
+	
+	if (SelectedItems.Num() > 0)
 	{
-		HorizontalBox->AddSlot()
-		.VAlign(VAlign_Center)
-		.HAlign(HAlign_Left)
-		.MaxWidth(24.0f)
-		.Padding(FMargin(4.f, 0.f, 4.f, 0.f))
-		[
-			SNew(SButton)
-			.ButtonStyle(FAppStyle::Get(), "NoBorder")
-			.ToolTipText(LOCTEXT("AddDependencyTooltip", "Add a new dependency to this variant"))
-			.OnClicked(this, &SVariantManager::OnAddDependencyClicked)
-			.ContentPadding(FMargin(2.0f, 1.0f))
-			.Content()
-			[
-				SNew(STextBlock)
-				.TextStyle(FAppStyle::Get(), "NormalText.Important")
-				.Font(FAppStyle::Get().GetFontStyle("FontAwesome.10"))
-				.Text(FEditorFontGlyphs::Plus)
-			]
-		];
+		FScopedTransaction Transaction(LOCTEXT("DeleteDependencies", "Delete variant dependencies"));
+
+		for (FVariantDependencyModelPtr Item : SelectedItems)
+		{
+			UVariant* ParentVariant = Item->ParentVariant.Get();
+			if (Item->Dependency && ParentVariant)
+			{
+				int32 DependencyIndex = INDEX_NONE;
+				for (int32 Index = 0; Index < ParentVariant->GetNumDependencies(); ++Index)
+				{
+					if (&ParentVariant->GetDependency(Index) == Item->Dependency)
+					{
+						DependencyIndex = Index;
+						break;
+					}
+				}
+
+				if (DependencyIndex != INDEX_NONE)
+				{
+
+
+					ParentVariant->DeleteDependency(DependencyIndex);
+				}
+			}
+		}
 	}
 
-	return
-	// Show a white spacer above the dependencies header row
-	SNew( SBorder )
-	.BorderImage( FAppStyle::GetBrush( "WhiteBrush" ) )
-	.BorderBackgroundColor( FVariantManagerStyle::Get()->GetColor( "VariantManager.Panels.LightBackgroundColor" ) )
-	.Padding( FMargin( 0.0f, RowType == ERightTreeRowType::DependenciesHeader ? FVariantManagerStyle::Get()->GetFloat( "VariantManager.Spacings.BorderThickness" ) : 0.0f, 0.0f, 0.0f ) )
-	[
-		SNew( SBox )
-		.HeightOverride( 26 )
-		[
-			SNew( SBorder )
-			.VAlign( VAlign_Center )
-			.HAlign( HAlign_Fill )
-			.BorderImage( FAppStyle::GetBrush( "Sequencer.AnimationOutliner.TopLevelBorder_Expanded" ) )
-			.BorderBackgroundColor( FVariantManagerStyle::Get()->GetColor( "VariantManager.Panels.HeaderBackgroundColor" ) )
-			.Padding( FMargin( 2.0f, 0.0f, 2.0f, 0.0f ) )
-			[
-				HorizontalBox
-			]
-		]
-	];
+	return FReply::Handled();
 }
 
 TSharedRef<SWidget> SVariantManager::GenerateRightTreePropertiesRowContent()
 {
-	float BorderThickness = FVariantManagerStyle::Get()->GetFloat("VariantManager.Spacings.BorderThickness");
+	float BorderThickness = FVariantManagerStyle::Get().GetFloat("VariantManager.Spacings.BorderThickness");
 
-	return SNew( SBorder )
-	.BorderImage( FAppStyle::GetBrush( "WhiteBrush" ) )
-	.BorderBackgroundColor( FVariantManagerStyle::Get()->GetColor( "VariantManager.Panels.ContentBackgroundColor" ) )
-	.Padding( FMargin( 0.0f, 0.0f, 0.0f, 2.0f ) )
+	return SNew( SBox )
 	[
 		// Actor column
 		SAssignNew( PropertiesSplitter, SSplitter )
@@ -3507,43 +3533,35 @@ TSharedRef<SWidget> SVariantManager::GenerateRightTreePropertiesRowContent()
 		.HitDetectionSplitterHandleSize( BorderThickness )
 
 		+ SSplitter::Slot()
-		.Value( PropertiesColumnSizeData.LeftColumnWidth )
-		.OnSlotResized( SSplitter::FOnSlotResized::CreateLambda( []( float InNewWidth ) {} ) )
 		[
 			SNew(SVerticalBox)
 			+ SVerticalBox::Slot()
-			.MaxHeight(VM_COMMON_HEADER_MAX_HEIGHT)
 			.AutoHeight()
-			.Padding(FMargin(0.0f, VM_COMMON_PADDING, 0.0f, VM_COMMON_PADDING))
 			[
-				SNew(SBox)
-				.HeightOverride(VM_COMMON_HEADER_MAX_HEIGHT)
-				.HAlign(HAlign_Center)
-				.VAlign(VAlign_Center)
+				SNew(SBorder)
+				.BorderImage(&HeaderStyle->NormalBrush)
+				.Padding(FMargin(10, 7, 8, 8))
 				[
 					SNew(SHorizontalBox)
-
 					+SHorizontalBox::Slot()
-					.AutoWidth()
+					.FillWidth(1.0f)
 					.VAlign(VAlign_Center)
+					.HAlign(HAlign_Left)
 					[
 						SNew(STextBlock)
-						.TextStyle(FAppStyle::Get(), "NormalText.Important")
 						.Text(LOCTEXT("ActorsText", "Actors"))
 					]
 
 					+SHorizontalBox::Slot()
 					.VAlign(VAlign_Center)
 					.HAlign(HAlign_Right)
-					.MaxWidth(24.0f)
 					.AutoWidth()
-					.Padding(FMargin(4.f, 0.f, 4.f, 0.f))
 					[
 						SNew(SButton)
-						.ButtonStyle(FAppStyle::Get(), "NoBorder")
-						.ToolTipText(LOCTEXT("AddActorPlusTooltip", "Add a new actor binding to selected variants"))
+						.ButtonStyle(FAppStyle::Get(), "HoverHintOnly")
 						.OnClicked(this, &SVariantManager::OnSummonAddActorMenu)
-						.ContentPadding(FMargin(2.0f, 1.0f))
+						.ContentPadding(FMargin(1, 0))
+						.ToolTipText(LOCTEXT("AddActorPlusTooltip", "Add a new actor binding to selected variants"))
 						.IsEnabled_Lambda([&VariantManagerPtr = VariantManagerPtr]() -> bool
 						{
 							if (TSharedPtr<FVariantManager> VariantManager = VariantManagerPtr.Pin())
@@ -3567,12 +3585,10 @@ TSharedRef<SWidget> SVariantManager::GenerateRightTreePropertiesRowContent()
 							}
 							return false;
 						})
-						.Content()
 						[
-							SNew(STextBlock)
-							.TextStyle(FAppStyle::Get(), "NormalText.Important")
-							.Font(FAppStyle::Get().GetFontStyle("FontAwesome.10"))
-							.Text(FEditorFontGlyphs::Plus)
+							SNew(SImage)
+							.Image(FAppStyle::Get().GetBrush("Icons.PlusCircle"))
+							.ColorAndOpacity(FSlateColor::UseForeground())
 						]
 					]
 				]
@@ -3584,7 +3600,6 @@ TSharedRef<SWidget> SVariantManager::GenerateRightTreePropertiesRowContent()
 			.FillHeight(1.0f)
 			[
 				SNew(SBox)
-				.MinDesiredHeight(30.0f)
 				[
 					ActorListView.ToSharedRef()
 				]
@@ -3593,100 +3608,57 @@ TSharedRef<SWidget> SVariantManager::GenerateRightTreePropertiesRowContent()
 
 		// Properties column
 		+ SSplitter::Slot()
-		.Value( TAttribute<float>::Create( TAttribute<float>::FGetter::CreateLambda( [this]()
-		{
-			return PropertiesNameColumnWidth + PropertiesValueColumnWidth;
-		} ) ) )
-		.OnSlotResized( SSplitter::FOnSlotResized::CreateLambda( [this]( float InNewWidth )
-		{
-			// InNewWidth is the new size of the property names + property values "combo", as they're a nested splitter
-			// We want to keep the property values column fixed in size though
-			PropertiesNameColumnWidth = FMath::Max(InNewWidth - PropertiesValueColumnWidth, 0.01f);
-		} ) )
 		[
 			SNew(SVerticalBox)
 			+ SVerticalBox::Slot()
-			.MaxHeight(VM_COMMON_HEADER_MAX_HEIGHT)
 			.AutoHeight()
-			.Padding(0.0f, VM_COMMON_PADDING, VM_COMMON_PADDING, VM_COMMON_PADDING)
 			[
-				// Headers
-				SNew(SSplitter)
-				.Orientation(Orient_Horizontal)
-				.PhysicalSplitterHandleSize( BorderThickness )
-				.HitDetectionSplitterHandleSize( BorderThickness )
-
-				+ SSplitter::Slot()
-				.Value( PropertiesColumnSizeData.MiddleColumnWidth )
-				.OnSlotResized( SSplitter::FOnSlotResized::CreateLambda( []( float InNewWidth ) {} ) )
+				SNew(SBorder)
+				.BorderImage(&HeaderStyle->NormalBrush)
+				.Padding(FMargin(10, 7, 8, 8))
 				[
-					SNew(SBox)
-					.HeightOverride(VM_COMMON_HEADER_MAX_HEIGHT)
-					.HAlign(HAlign_Center)
+					SNew(SHorizontalBox)
+					+SHorizontalBox::Slot()
+					.FillWidth(1.0f)
 					.VAlign(VAlign_Center)
+					.HAlign(HAlign_Left)
 					[
-						SNew(SHorizontalBox)
+						SNew(STextBlock)
+						.Text(LOCTEXT("PropertiesText", "Properties"))
+					]
 
-						+SHorizontalBox::Slot()
-						.AutoWidth()
-						.VAlign(VAlign_Center)
+					+SHorizontalBox::Slot()
+					.VAlign(VAlign_Center)
+					.HAlign(HAlign_Right)
+					.AutoWidth()
+					[
+						SNew(SButton)
+						.ButtonStyle(FAppStyle::Get(), "HoverHintOnly")
+						.ContentPadding(FMargin(1, 0))
+						.ToolTipText(LOCTEXT("CapturePropertiesPlusTooltip", "Capture properties from the selected actor bindings"))
+						.OnClicked_Lambda([this]
+						{
+							CaptureNewPropertiesFromSelectedActors();
+							return FReply::Handled();
+						})
 						[
-							SNew(STextBlock)
-							.TextStyle(FAppStyle::Get(), "NormalText.Important")
-							.Text(LOCTEXT("PropertiesText", "Properties"))
+							SNew(SImage)
+							.Image(FAppStyle::Get().GetBrush("Icons.PlusCircle"))
+							.ColorAndOpacity(FSlateColor::UseForeground())
 						]
-
-						+SHorizontalBox::Slot()
-						.VAlign(VAlign_Center)
-						.HAlign(HAlign_Right)
-						.MaxWidth(24.0f)
-						.AutoWidth()
-						.Padding(FMargin(4.f, 0.f, 4.f, 0.f))
-						[
-							SNew(SButton)
-							.ButtonStyle(FAppStyle::Get(), "NoBorder")
-							.ToolTipText(LOCTEXT("CapturePropertiesPlusTooltip", "Capture properties from the selected actor bindings"))
-							.OnClicked_Lambda([this]
-							{
-								CaptureNewPropertiesFromSelectedActors();
-								return FReply::Handled();
-							})
-							.ContentPadding(FMargin(2.0f, 1.0f))
-							.IsEnabled_Lambda([&VariantManagerPtr = VariantManagerPtr]() -> bool
+						.IsEnabled_Lambda([&VariantManagerPtr = VariantManagerPtr]() -> bool
+						{
+							if (TSharedPtr<FVariantManager> VariantManager = VariantManagerPtr.Pin())
 							{
-								if (TSharedPtr<FVariantManager> VariantManager = VariantManagerPtr.Pin())
-								{
-									return VariantManager->GetSelection().GetSelectedActorNodes().Num() > 0;
-								}
-								return false;
-							})
-							.Content()
-							[
-								SNew(STextBlock)
-								.TextStyle(FAppStyle::Get(), "NormalText.Important")
-								.Font(FAppStyle::Get().GetFontStyle("FontAwesome.10"))
-								.Text(FEditorFontGlyphs::Plus)
-							]
-						]
-					]
-				]
-				+ SSplitter::Slot()
-				.Value( PropertiesColumnSizeData.RightColumnWidth )
-				.OnSlotResized( PropertiesColumnSizeData.OnSecondSplitterChanged )
-				[
-					SNew(SBox)
-					.HeightOverride(VM_COMMON_HEADER_MAX_HEIGHT)
-					.HAlign(HAlign_Center)
-					.VAlign(VAlign_Center)
-					[
-						SNew(STextBlock)
-						.TextStyle(FAppStyle::Get(), "NormalText.Important")
-						.Text(LOCTEXT("PropertiesValuesText", "Values"))
+								return VariantManager->GetSelection().GetSelectedActorNodes().Num() > 0;
+							}
+							return false;
+						})
 					]
 				]
+				
 			]
 			+ SVerticalBox::Slot()
-			.Padding(0.0f, 0.0f, VM_COMMON_PADDING, 0.0f)
 			.VAlign(VAlign_Fill)
 			.FillHeight(1.0f)
 			[
@@ -3703,130 +3675,57 @@ TSharedRef<SWidget> SVariantManager::GenerateRightTreePropertiesRowContent()
 
 TSharedRef<SWidget> SVariantManager::GenerateRightTreeDependenciesRowContent()
 {
-	float BorderThickness = FVariantManagerStyle::Get()->GetFloat( "VariantManager.Spacings.BorderThickness" );
-
-	return SNew( SBorder )
-	.BorderImage( FAppStyle::GetBrush( "WhiteBrush" ) )
-	.BorderBackgroundColor( FVariantManagerStyle::Get()->GetColor( "VariantManager.Panels.ContentBackgroundColor" ) )
-	.Padding( FMargin( BorderThickness, 0.0f, BorderThickness, 1.0f ) )
-	[
-		SNew(SVerticalBox)
-		.IsEnabled_Lambda([this]()
-		{
-			TSharedPtr<FVariantManager> PinnedVariantManager = VariantManagerPtr.Pin();
-			if (PinnedVariantManager.IsValid())
+	return 
+		SNew( SBox )
+		[
+			SNew(SVerticalBox)
+			.IsEnabled_Lambda([this]()
 			{
-				TArray<UVariant*> SelectedVariants;
-				TArray<UVariantSet*> SelectedVariantSets;
-				PinnedVariantManager->GetSelection().GetSelectedVariantsAndVariantSets(SelectedVariants, SelectedVariantSets);
+				TSharedPtr<FVariantManager> PinnedVariantManager = VariantManagerPtr.Pin();
+				if (PinnedVariantManager.IsValid())
+				{
+					TArray<UVariant*> SelectedVariants;
+					TArray<UVariantSet*> SelectedVariantSets;
+					PinnedVariantManager->GetSelection().GetSelectedVariantsAndVariantSets(SelectedVariants, SelectedVariantSets);
 
-				return SelectedVariants.Num() == 1;
-			}
-			return false;
-		})
+					return SelectedVariants.Num() == 1;
+				}
+				return false;
+			})
 
-		// Headers
-		+ SVerticalBox::Slot()
-		.MaxHeight(VM_COMMON_HEADER_MAX_HEIGHT)
-		.AutoHeight()
-		.Padding(0.0f, VM_COMMON_PADDING, VM_COMMON_PADDING, VM_COMMON_PADDING)
-		[
-			SNew(SSplitter)
-			.Orientation(Orient_Horizontal)
-			.PhysicalSplitterHandleSize( BorderThickness )
-			.HitDetectionSplitterHandleSize( BorderThickness )
-
-			+SSplitter::Slot()
-			.Value(DependenciesColumnSizeData.LeftColumnWidth)
-			.OnSlotResized(SSplitter::FOnSlotResized::CreateLambda([](float InNewWidth)
-			{
-				//This has to be bound or the splitter will take it upon itself to determine the size
-				//We do nothing here because it is handled by the column size data
-			}))
+			+ SVerticalBox::Slot()
 			[
-				SNew(SBox)
-				.HeightOverride(VM_COMMON_HEADER_MAX_HEIGHT)
-				.HAlign(HAlign_Center)
-				.VAlign(VAlign_Center)
-				[
-					SNew( STextBlock )
-					.TextStyle( FAppStyle::Get(), "NormalText.Important" )
-					.Text( LOCTEXT( "DependenciesVariantSetText", "Variant Set" ) )
-				]
+				SNew(SSeparator)
+				.Orientation(EOrientation::Orient_Horizontal)
+				.Thickness(3)
+				.SeparatorImage(&SplitterStyle->HandleNormalBrush)
 			]
-			+ SSplitter::Slot()
-			.Value( DependenciesColumnSizeData.MiddleColumnWidth )
-			.OnSlotResized( DependenciesColumnSizeData.OnFirstSplitterChanged )
+
+			// Dependencies
+			+ SVerticalBox::Slot()
+			.VAlign(VAlign_Top)
+			.AutoHeight()
 			[
-				SNew(SBox)
-				.HeightOverride(VM_COMMON_HEADER_MAX_HEIGHT)
-				.HAlign(HAlign_Center)
-				.VAlign(VAlign_Center)
-				.Padding(FMargin(0, 0, 0, 0))
+				SNew(SBox) // Very important to prevent the tree from expanding freely
 				[
-					SNew(STextBlock)
-					.TextStyle(FAppStyle::Get(), "NormalText.Important")
-					.Text(LOCTEXT("DependenciesVariantText", "Variant"))
+					SAssignNew(DependenciesList, SListView<FVariantDependencyModelPtr>)
+					.SelectionMode(ESelectionMode::Multi)
+					.ListItemsSource(&DisplayedDependencies)
+					.OnGenerateRow(this, &SVariantManager::GenerateDependencyRow, true)
+					.HeaderRow(SNew(SHeaderRow)
+						+ SHeaderRow::Column(SDependencyRow::VisibilityColumn)
+						.DefaultLabel(LOCTEXT("", ""))
+						.FixedWidth(30)
+						+ SHeaderRow::Column(SDependencyRow::VariantSetColumn)
+						.DefaultLabel(LOCTEXT("DependenciesVariantSetText", "Variant Set"))
+						.FillWidth(0.5f)
+						+ SHeaderRow::Column(SDependencyRow::VariantColumn)
+						.DefaultLabel(LOCTEXT("DependenciesVariantText", "Variant"))
+						.FillWidth(0.5f))
+					.Visibility(EVisibility::Visible)
 				]
 			]
-			+ SSplitter::Slot()
-			.Value( DependenciesColumnSizeData.RightColumnWidth )
-			.OnSlotResized( DependenciesColumnSizeData.OnSecondSplitterChanged )
-			[
-				SNullWidget::NullWidget
-			]
-		]
-
-		// Dependencies
-		+ SVerticalBox::Slot()
-		.Padding(0.0f, 0.0f, VM_COMMON_PADDING, 1.0f)
-		.VAlign(VAlign_Top)
-		.AutoHeight()
-		[
-			SNew( SBox )
-			.MinDesiredHeight( 30.0f )
-			[
-				SAssignNew(DependenciesList, SListView<FVariantDependencyModelPtr>)
-				.SelectionMode(ESelectionMode::None)
-				.ListItemsSource(&DisplayedDependencies)
-				.OnGenerateRow(this, &SVariantManager::GenerateDependencyRow, true)
-				.Visibility(EVisibility::Visible)
-			]
-		]
-
-		// Separator
-		+ SVerticalBox::Slot()
-		.Padding(0.0f, 10.0f, VM_COMMON_PADDING, 10.0f)
-		.VAlign(VAlign_Top)
-		.MaxHeight(24)
-		[
-			SNew(SSeparator)
-			.SeparatorImage(FAppStyle::GetBrush("ThinLine.Horizontal"))
-			.Thickness(4.f)
-			.Orientation(EOrientation::Orient_Horizontal)
-			.ColorAndOpacity(FLinearColor(0.1, 0.1, 0.1, 1.0))
-			.Visibility_Lambda([this]()
-			{
-				return DisplayedDependents.Num() > 0 ? EVisibility::Visible : EVisibility::Collapsed;
-			})
-		]
-
-		// Dependents
-		+ SVerticalBox::Slot()
-		.Padding(0.0f, 0.0f, VM_COMMON_PADDING, 1.0f)
-		.VAlign(VAlign_Top)
-		.AutoHeight()
-		[
-			SAssignNew(DependentsList, SListView<FVariantDependencyModelPtr>)
-			.SelectionMode(ESelectionMode::None)
-			.ListItemsSource(&DisplayedDependents)
-			.OnGenerateRow(this, &SVariantManager::GenerateDependencyRow, false)
-			.Visibility_Lambda( [ this ]()
-			{
-				return DisplayedDependents.Num() > 0 ? EVisibility::Visible : EVisibility::Collapsed;
-			})
-		]
-	];
+		];
 }
 
 TSharedRef<ITableRow> SVariantManager::GenerateDependencyRow(FVariantDependencyModelPtr Dependency, const TSharedRef<STableViewBase>& OwnerTable, bool bInteractionEnabled)
@@ -3845,7 +3744,8 @@ TSharedRef<ITableRow> SVariantManager::GenerateDependencyRow(FVariantDependencyM
 		return SNew(STableRow<FVariantDependencyModelPtr>, OwnerTable);
 	}
 
-	return SNew(SDependencyRow, OwnerTable, GetDependenciesColumnSizeData(), Dependency, bInteractionEnabled);
+	return SNew(SDependencyRow, OwnerTable, Dependency, bInteractionEnabled)
+		.IsEnabled(!Dependency->bIsDependent && !Dependency->bIsDivider);
 }
 
 void SVariantManager::OnEditorSelectionChanged(UObject* NewSelection)
diff --git a/Engine/Plugins/Enterprise/VariantManager/Source/VariantManager/Private/SVariantManager.h b/Engine/Plugins/Enterprise/VariantManager/Source/VariantManager/Private/SVariantManager.h
--- a/Engine/Plugins/Enterprise/VariantManager/Source/VariantManager/Private/SVariantManager.h
+++ b/Engine/Plugins/Enterprise/VariantManager/Source/VariantManager/Private/SVariantManager.h
@@ -35,25 +35,21 @@ using FVariantDependencyModelPtr = TSharedPtr<FVariantDependencyModel>;
 // Convenience struct to save/load how the user configured the main splitters
 struct FSplitterValues
 {
+	float PropertyNameColumn = 0.5f;
+	float PropertyValueColumn = 0.5f;
 	float VariantColumn = 0.25f;
-	float ActorColumn = 0.25f;
-	float PropertyNameColumn = 0.25f;
-	float DependenciesVariantSetsColumn = 0.33f;
-	float DependenciesVariantColumn = 0.33f;
 
 	FSplitterValues(){};
 	FSplitterValues(FString& InSerialized);
 	FString ToString();
 };
 
-struct FColumnSizeData
+struct FVariantManagerPropertiesColumnSizeData
 {
-	TAttribute<float> LeftColumnWidth;
-	TAttribute<float> MiddleColumnWidth;
-	TAttribute<float> RightColumnWidth;
-	// There are three columns, but only two splitters
-	SSplitter::FOnSlotResized OnFirstSplitterChanged;
-	SSplitter::FOnSlotResized OnSecondSplitterChanged;
+	TAttribute<float> NameColumnWidth;
+	TAttribute<float> ValueColumnWidth;
+	SSplitter::FOnSlotResized OnSplitterNameColumnChanged;
+	SSplitter::FOnSlotResized OnSplitterValueColumnChanged;
 };
 
 class SVariantManager
@@ -170,14 +166,11 @@ public:
 	void SortDisplayNodes(TArray<TSharedRef<FVariantManagerDisplayNode>>& DisplayNodes);
 
 	TSharedRef<SWidget> MakeAddButton();
-	FColumnSizeData& GetPropertiesColumnSizeData()
+
+	FVariantManagerPropertiesColumnSizeData& GetPropertiesColumnSizeData()
 	{
 		return PropertiesColumnSizeData;
 	}
-	FColumnSizeData& GetDependenciesColumnSizeData()
-	{
-		return DependenciesColumnSizeData;
-	}
 
 	TSharedRef<ITableRow> MakeCapturedPropertyRow(TSharedPtr<FVariantManagerPropertyNode> Item, const TSharedRef<STableViewBase>& OwnerTable);
 	TSharedPtr<SWidget> OnPropertyListContextMenuOpening();
@@ -214,19 +207,11 @@ public:
 	FReply OnSummonAddActorMenu();
 	FReply OnAddDependencyClicked();
 
+	FReply OnDeleteSelectedDependencies();
+
 	// Callbacks for property ColumnSizeData
-	float OnGetPropertiesActorColumnWidth() const { return 1.0f - PropertiesNameColumnWidth - PropertiesValueColumnWidth; }
 	float OnGetPropertiesNameColumnWidth() const { return PropertiesNameColumnWidth; }
 	float OnGetPropertiesValueColumnWidth() const { return PropertiesValueColumnWidth; }
-	void OnSetPropertiesNameColumnWidth( float InWidth ) { PropertiesNameColumnWidth = InWidth; }
-	void OnSetPropertiesValueColumnWidth( float InWidth ) { PropertiesValueColumnWidth = InWidth; }
-
-	// Callbacks for dependencies ColumnSizeData
-	float OnGetDependenciesVariantSetColumnWidth() const { return 1.0f - DependenciesVariantColumnWidth - DependenciesControlColumnWidth; }
-	float OnGetDependenciesVariantColumnWidth() const { return DependenciesVariantColumnWidth; }
-	float OnGetDependenciesControlColumnWidth() const { return DependenciesControlColumnWidth; }
-	void OnSetDependenciesVariantColumnWidth( float InWidth ) { DependenciesVariantColumnWidth = InWidth; }
-	void OnSetDependenciesControlColumnWidth( float InWidth ) { DependenciesControlColumnWidth = InWidth; }
 
 	void OnObjectTransacted(UObject* Object, const class FTransactionObjectEvent& Event);
 	void OnObjectPropertyChanged(UObject* Object, struct FPropertyChangedEvent& Event);
@@ -267,9 +252,6 @@ private:
 	TSharedPtr<SListView<FVariantDependencyModelPtr>> DependenciesList;
 	TArray<FVariantDependencyModelPtr> DisplayedDependencies;
 
-	TSharedPtr<SListView<FVariantDependencyModelPtr>> DependentsList;
-	TArray<FVariantDependencyModelPtr> DisplayedDependents;
-
 	// We use paths here to avoid having to check if the bindings are resolved
 	TSet<FString> CachedSelectedActorPaths;
 	TSet<FString> CachedDisplayedActorPaths;
@@ -285,14 +267,10 @@ private:
 
 	bool bAutoCaptureProperties = false;
 
-	FColumnSizeData PropertiesColumnSizeData;
+	FVariantManagerPropertiesColumnSizeData PropertiesColumnSizeData;
 	float PropertiesNameColumnWidth;
 	float PropertiesValueColumnWidth;
 
-	FColumnSizeData DependenciesColumnSizeData;
-	float DependenciesVariantColumnWidth;
-	float DependenciesControlColumnWidth;
-
 	FSplitterValues SplitterValues;
 
 	FDelegateHandle OnObjectTransactedHandle;
@@ -328,4 +306,7 @@ private:
 	TArray<SVariantManager::FCachedPropertyPath> CachedPropertyPathStack;
 
 	bool bRespondToEditorSelectionEvents = true;
+
+	const FTableColumnHeaderStyle* HeaderStyle = nullptr;
+	const FSplitterStyle* SplitterStyle = nullptr;
 };
diff --git a/Engine/Plugins/Enterprise/VariantManager/Source/VariantManager/Private/SVariantManagerTableRow.cpp b/Engine/Plugins/Enterprise/VariantManager/Source/VariantManager/Private/SVariantManagerTableRow.cpp
--- a/Engine/Plugins/Enterprise/VariantManager/Source/VariantManager/Private/SVariantManagerTableRow.cpp
+++ b/Engine/Plugins/Enterprise/VariantManager/Source/VariantManager/Private/SVariantManagerTableRow.cpp
@@ -24,11 +24,13 @@ void SVariantManagerTableRow::Construct(const FArguments& InArgs, const TSharedR
 
 	STableRow::Construct(
 		STableRow::FArguments()
-			.OnDragDetected(this, &SVariantManagerTableRow::DragDetected)
-			.OnCanAcceptDrop(this, &SVariantManagerTableRow::CanAcceptDrop)
-			.OnAcceptDrop(this, &SVariantManagerTableRow::AcceptDrop)
-			.OnDragLeave(this, &SVariantManagerTableRow::DragLeave)
-			.ShowSelection(bIsSelectable),
+		.Padding(0)
+		.Style(InNode->GetRowStyle())
+		.OnDragDetected(this, &SVariantManagerTableRow::DragDetected)
+		.OnCanAcceptDrop(this, &SVariantManagerTableRow::CanAcceptDrop)
+		.OnAcceptDrop(this, &SVariantManagerTableRow::AcceptDrop)
+		.OnDragLeave(this, &SVariantManagerTableRow::DragLeave)
+		.ShowSelection(bIsSelectable),
 		OwnerTableView);
 
 	SetRowContent(InNode->GetCustomOutlinerContent(SharedThis(this)));
diff --git a/Engine/Plugins/Enterprise/VariantManager/Source/VariantManager/Private/VariantManagerModule.cpp b/Engine/Plugins/Enterprise/VariantManager/Source/VariantManager/Private/VariantManagerModule.cpp
--- a/Engine/Plugins/Enterprise/VariantManager/Source/VariantManager/Private/VariantManagerModule.cpp
+++ b/Engine/Plugins/Enterprise/VariantManager/Source/VariantManager/Private/VariantManagerModule.cpp
@@ -29,8 +29,6 @@ public:
 	{
 		FVariantManagerEditorCommands::Register();
 
-		FVariantManagerStyle::Initialize();
-
 		FLevelEditorModule& LevelEditorModule = FModuleManager::LoadModuleChecked<FLevelEditorModule>(TEXT("LevelEditor"));
 
 		// Register a delegate to detect whenever we should open an editor for a LevelVariantSets asset, and relay the
@@ -64,8 +62,6 @@ public:
 
 		FVariantManagerUtils::UnregisterForHotReload();
 
-		FVariantManagerStyle::Shutdown();
-
 		FVariantManagerEditorCommands::Unregister();
 	}
 
@@ -89,12 +85,12 @@ public:
 			UnregisterTabSpawner( TabManager );
 		}
 
-		const FSlateIcon LayersIcon( FAppStyle::GetAppStyleSetName(), "LevelEditor.Tabs.Layers" );
+		const FName VariantManagerStyleSetName = FVariantManagerStyle::Get().GetStyleSetName();
 
 		TabManager->RegisterTabSpawner( FLevelVariantSetsEditorToolkit::GetVariantManagerTabID(), FOnSpawnTab::CreateStatic( &FVariantManagerModule::HandleTabManagerSpawnTab ) )
 			.SetDisplayName( LOCTEXT("VariantManagerMainTab", "Variant Manager") )
 			.SetGroup( WorkspaceMenu::GetMenuStructure().GetLevelEditorCategory() )
-			.SetIcon( LayersIcon );
+			.SetIcon( FSlateIcon(VariantManagerStyleSetName, "VariantManager.Icon") );
 	}
 
 	static void UnregisterTabSpawner( const TSharedPtr< FTabManager >& TabManager )
diff --git a/Engine/Plugins/Enterprise/VariantManager/Source/VariantManager/Private/VariantManagerStyle.cpp b/Engine/Plugins/Enterprise/VariantManager/Source/VariantManager/Private/VariantManagerStyle.cpp
--- a/Engine/Plugins/Enterprise/VariantManager/Source/VariantManager/Private/VariantManagerStyle.cpp
+++ b/Engine/Plugins/Enterprise/VariantManager/Source/VariantManager/Private/VariantManagerStyle.cpp
@@ -8,91 +8,54 @@
 #include "Styling/CoreStyle.h"
 #include "Styling/SlateStyleRegistry.h"
 #include "Styling/SlateTypes.h"
+#include "Styling/SlateStyleMacros.h"
 
-#define IMAGE_PLUGIN_BRUSH( RelativePath, ... ) FSlateImageBrush( FVariantManagerStyle::InContent( RelativePath, ".png" ), __VA_ARGS__ )
-#define IMAGE_BRUSH(RelativePath, ...) FSlateImageBrush(StyleSet->RootToContentDir(RelativePath, TEXT(".png")), __VA_ARGS__)
-#define BOX_BRUSH(RelativePath, ...) FSlateBoxBrush(StyleSet->RootToContentDir(RelativePath, TEXT(".png")), __VA_ARGS__)
-#define DEFAULT_FONT(...) FCoreStyle::GetDefaultFontStyle(__VA_ARGS__)
-
-FString FVariantManagerStyle::InContent(const FString& RelativePath, const ANSICHAR* Extension)
-{
-	static FString ContentDir = IPluginManager::Get().FindPlugin(TEXT("VariantManager"))->GetContentDir();
-	return (ContentDir / RelativePath) + Extension;
-}
-
-TSharedPtr< FSlateStyleSet > FVariantManagerStyle::StyleSet = nullptr;
-TSharedPtr< class ISlateStyle > FVariantManagerStyle::Get() { return StyleSet; }
-
-FName FVariantManagerStyle::GetStyleSetName()
+FVariantManagerStyle::FVariantManagerStyle() : FSlateStyleSet("VariantManagerEditorStyle")
 {
-	static FName PaperStyleName(TEXT("VariantManager"));
-	return PaperStyleName;
-}
+	SetContentRoot(IPluginManager::Get().FindPlugin(TEXT("VariantManager"))->GetContentDir());
 
-BEGIN_SLATE_FUNCTION_BUILD_OPTIMIZATION
-
-void FVariantManagerStyle::Initialize()
-{
 	// Const icon sizes
 	const FVector2D Icon8x8(8.0f, 8.0f);
 	const FVector2D Icon16x16(16.0f, 16.0f);
 	const FVector2D Icon20x20(20.0f, 20.0f);
 	const FVector2D Icon40x40(40.0f, 40.0f);
 
-	// Only register once
-	if (StyleSet.IsValid())
-	{
-		return;
-	}
-
-	StyleSet = MakeShareable(new FSlateStyleSet(GetStyleSetName()));
-	StyleSet->SetContentRoot(FPaths::EngineContentDir() / TEXT("Editor/Slate"));
-	StyleSet->SetCoreContentRoot(FPaths::EngineContentDir() / TEXT("Slate"));
+	/** Color used for the background of the entire variant manager as well as the spacer border */
+	Set( "VariantManager.Panels.LightBackgroundColor", FLinearColor( FColor( 96, 96, 96, 255 ) ) );
 
-	{
-		/** Color used for the background of the entire variant manager as well as the spacer border */
-		StyleSet->Set( "VariantManager.Panels.LightBackgroundColor", FLinearColor( FColor( 96, 96, 96, 255 ) ) );
+	/** Color used as background for variant nodes, and background of properties and dependencies panels */
+	Set( "VariantManager.Panels.ContentBackgroundColor", FLinearColor( FColor( 62, 62, 62, 255 ) ) );
 
-		/** Color used as background for variant nodes, and background of properties and dependencies panels */
-		StyleSet->Set( "VariantManager.Panels.ContentBackgroundColor", FLinearColor( FColor( 62, 62, 62, 255 ) ) );
+	/** Color used for background of variant set nodes and panel headers, like Properties or Dependencies headers */
+	Set( "VariantManager.Panels.HeaderBackgroundColor", FLinearColor( FColor( 48, 48, 48, 255 ) ) );
 
-		/** Color used for background of variant set nodes and panel headers, like Properties or Dependencies headers */
-		StyleSet->Set( "VariantManager.Panels.HeaderBackgroundColor", FLinearColor( FColor( 48, 48, 48, 255 ) ) );
+	/** Thickness of the light border around the entire variant manager tab and between items */
+	Set( "VariantManager.Spacings.BorderThickness", 4.0f );
 
-		/** Thickness of the light border around the entire variant manager tab and between items */
-		StyleSet->Set( "VariantManager.Spacings.BorderThickness", 4.0f );
+	/** The amount to indent child nodes of the layout tree */
+	Set( "VariantManager.Spacings.IndentAmount", 10.0f );
 
-		/** The amount to indent child nodes of the layout tree */
-		StyleSet->Set( "VariantManager.Spacings.IndentAmount", 10.0f );
+	Set( "VariantManager.Icon", new IMAGE_BRUSH_SVG("VariantManager", Icon16x16) );
 
-		/** Style used for the auto-capture button */
-		const FCheckBoxStyle Style = FCheckBoxStyle()
-			.SetCheckBoxType( ESlateCheckBoxType::ToggleButton )
-			.SetUncheckedImage( BOX_BRUSH( "/Sequencer/Transport_Bar/Record_24x_OFF", 0.0f, FSlateColor( FLinearColor::Gray ) ) )
-			.SetUncheckedHoveredImage( BOX_BRUSH( "/Sequencer/Transport_Bar/Record_24x_OFF", 0.0f, FSlateColor( FLinearColor::White ) ) )
-			.SetUncheckedPressedImage( BOX_BRUSH( "/Sequencer/Transport_Bar/Record_24x_OFF", 0.0f, FSlateColor( FLinearColor(0.75f, 0.75f, 0.75f) ) ) )
-			.SetCheckedImage( BOX_BRUSH( "/Sequencer/Transport_Bar/Recording_24x", 0.0f, FSlateColor( FLinearColor::Gray ) ) )
-			.SetCheckedHoveredImage( BOX_BRUSH( "/Sequencer/Transport_Bar/Recording_24x", 0.0f, FSlateColor( FLinearColor::White ) ) )
-			.SetCheckedPressedImage( BOX_BRUSH( "/Sequencer/Transport_Bar/Recording_24x", 0.0f, FSlateColor( FLinearColor( 0.75f, 0.75f, 0.75f ) ) ) );
-		StyleSet->Set( "AutoCaptureCheckbox", Style );
-	}
+	/** Style used for the auto-capture button */
+	const FCheckBoxStyle AutoCaptureButtonStyle = FCheckBoxStyle()
+		.SetCheckBoxType(ESlateCheckBoxType::ToggleButton)
+		.SetUncheckedImage(IMAGE_BRUSH("AutoCapture", Icon20x20))
+		.SetUncheckedHoveredImage(IMAGE_BRUSH("AutoCapture", Icon20x20, FLinearColor(0.5f, 0.5f, 0.5f, 1.0f)))
+		.SetUncheckedPressedImage(IMAGE_BRUSH("AutoCapture", Icon20x20, FLinearColor(0.5f, 0.5f, 0.5f, 1.0f)))
+		.SetCheckedImage(IMAGE_BRUSH("AutoCaptureChecked", Icon20x20))
+		.SetCheckedHoveredImage(IMAGE_BRUSH("AutoCaptureChecked", Icon20x20, FLinearColor(0.5f, 0.5f, 0.5f, 1.0f)))
+		.SetCheckedPressedImage(IMAGE_BRUSH("AutoCaptureChecked", Icon20x20, FLinearColor(0.5f, 0.5f, 0.5f, 1.0f)));
+	Set("VariantManager.AutoCapture", AutoCaptureButtonStyle);
 
-	FSlateStyleRegistry::RegisterSlateStyle(*StyleSet.Get());
-};
+	const FCheckBoxStyle RadioButtonStyle = FCheckBoxStyle()
+		.SetUncheckedImage(IMAGE_BRUSH("VariantUnchecked", Icon16x16))
+		.SetUncheckedHoveredImage(IMAGE_BRUSH("VariantUnchecked", Icon16x16, FLinearColor(0.5f, 0.5f, 0.5f, 1.0f)))
+		.SetUncheckedPressedImage(IMAGE_BRUSH("VariantUnchecked", Icon16x16, FLinearColor(0.5f, 0.5f, 0.5f, 1.0f)))
+		.SetCheckedImage(IMAGE_BRUSH("VariantChecked", Icon16x16))
+		.SetCheckedHoveredImage(IMAGE_BRUSH("VariantChecked", Icon16x16, FLinearColor(0.5f, 0.5f, 0.5f, 1.0f)))
+		.SetCheckedPressedImage(IMAGE_BRUSH("VariantChecked", Icon16x16, FLinearColor(0.5f, 0.5f, 0.5f, 1.0f)));
+	Set("VariantManager.VariantRadioButton", RadioButtonStyle);
 
-END_SLATE_FUNCTION_BUILD_OPTIMIZATION
-
-#undef IMAGE_PLUGIN_BRUSH
-#undef IMAGE_BRUSH
-#undef BOX_BRUSH
-#undef DEFAULT_FONT
-
-void FVariantManagerStyle::Shutdown()
-{
-	if (StyleSet.IsValid())
-	{
-		FSlateStyleRegistry::UnRegisterSlateStyle(*StyleSet.Get());
-		ensure(StyleSet.IsUnique());
-		StyleSet.Reset();
-	}
+	FSlateStyleRegistry::RegisterSlateStyle(*this);
 }
diff --git a/Engine/Plugins/Enterprise/VariantManager/Source/VariantManager/Private/VariantManagerStyle.h b/Engine/Plugins/Enterprise/VariantManager/Source/VariantManager/Private/VariantManagerStyle.h
--- a/Engine/Plugins/Enterprise/VariantManager/Source/VariantManager/Private/VariantManagerStyle.h
+++ b/Engine/Plugins/Enterprise/VariantManager/Source/VariantManager/Private/VariantManagerStyle.h
@@ -5,20 +5,21 @@
 #include "CoreMinimal.h"
 #include "Styling/ISlateStyle.h"
 #include "Styling/SlateStyle.h"
+#include "Styling/SlateStyleRegistry.h"
 
-class FVariantManagerStyle
+class FVariantManagerStyle final : public FSlateStyleSet
 {
 public:
-	static void Initialize();
-	static void Shutdown();
+	FVariantManagerStyle();
 
-	static TSharedPtr< class ISlateStyle > Get();
+	~FVariantManagerStyle()
+	{
+		FSlateStyleRegistry::UnRegisterSlateStyle(*this);
+	}
 
-	static FName GetStyleSetName();
-
-private:
-	static FString InContent(const FString& RelativePath, const ANSICHAR* Extension);
-
-private:
-	static TSharedPtr< class FSlateStyleSet > StyleSet;
+	static FVariantManagerStyle& Get()
+	{
+		static FVariantManagerStyle Inst;
+		return Inst;
+	}
 };
diff --git a/Engine/Plugins/Enterprise/VariantManager/Source/VariantManager/VariantManager.Build.cs b/Engine/Plugins/Enterprise/VariantManager/Source/VariantManager/VariantManager.Build.cs
--- a/Engine/Plugins/Enterprise/VariantManager/Source/VariantManager/VariantManager.Build.cs
+++ b/Engine/Plugins/Enterprise/VariantManager/Source/VariantManager/VariantManager.Build.cs
@@ -36,6 +36,7 @@ namespace UnrealBuildTool.Rules
 					"Slate",
 					"SlateCore",
 					"ToolMenus",
+					"ToolWidgets",
 					"VariantManagerContentEditor",
 					"WorkspaceMenuStructure",
 				}
