//UE5/Main - Build Engine Localization - SVariantManager.cpp Warning: LOCTEXT macro has an empty identifier and cannot be gathered.

diff --git a/Engine/Plugins/Enterprise/VariantManager/Source/VariantManager/Private/SVariantManager.cpp b/Engine/Plugins/Enterprise/VariantManager/Source/VariantManager/Private/SVariantManager.cpp
--- a/Engine/Plugins/Enterprise/VariantManager/Source/VariantManager/Private/SVariantManager.cpp
+++ b/Engine/Plugins/Enterprise/VariantManager/Source/VariantManager/Private/SVariantManager.cpp
@@ -3714,7 +3714,7 @@ TSharedRef<SWidget> SVariantManager::GenerateRightTreeDependenciesRowContent()
 					.OnGenerateRow(this, &SVariantManager::GenerateDependencyRow, true)
 					.HeaderRow(SNew(SHeaderRow)
 						+ SHeaderRow::Column(SDependencyRow::VisibilityColumn)
-						.DefaultLabel(LOCTEXT("", ""))
+						.DefaultLabel(LOCTEXT("VariantManagerEmptyText", ""))
 						.FixedWidth(30)
 						+ SHeaderRow::Column(SDependencyRow::VariantSetColumn)
 						.DefaultLabel(LOCTEXT("DependenciesVariantSetText", "Variant Set"))
