Improve Variant manager UI

Tweaks

diff --git a/Engine/Plugins/Enterprise/VariantManager/Source/VariantManager/Private/DisplayNodes/VariantManagerVariantNode.cpp b/Engine/Plugins/Enterprise/VariantManager/Source/VariantManager/Private/DisplayNodes/VariantManagerVariantNode.cpp
--- a/Engine/Plugins/Enterprise/VariantManager/Source/VariantManager/Private/DisplayNodes/VariantManagerVariantNode.cpp
+++ b/Engine/Plugins/Enterprise/VariantManager/Source/VariantManager/Private/DisplayNodes/VariantManagerVariantNode.cpp
@@ -74,7 +74,7 @@ TSharedRef<SWidget> FVariantManagerVariantNode::GetCustomOutlinerContent(TShared
 
 	return
 	SNew(SBox)
-	.HeightOverride(78)
+	.HeightOverride(40)
 	[
 		SNew(SBorder)
 		.VAlign(VAlign_Center)
@@ -105,7 +105,7 @@ TSharedRef<SWidget> FVariantManagerVariantNode::GetCustomOutlinerContent(TShared
 			.VAlign(VAlign_Center)
 			.HAlign(HAlign_Left)
 			.FillWidth(1.0f)
-			.Padding(FMargin(0.f, 0.f, 4.f, 0.f))
+			.Padding(FMargin(4.f, 0.f, 4.f, 0.f))
 			[
 				EditableLabel.ToSharedRef()
 			]
@@ -510,8 +510,8 @@ TSharedRef<SWidget> FVariantManagerVariantNode::GetThumbnailWidget()
 	}
 
 	return SNew(SBox)
-		.WidthOverride(64)
-		.HeightOverride(64)
+		.WidthOverride(32)
+		.HeightOverride(32)
 		[
 			ThumbnailWidget.IsValid() ? ThumbnailWidget.ToSharedRef() : SNullWidget::NullWidget
 		];
diff --git a/Engine/Plugins/Enterprise/VariantManager/Source/VariantManager/Private/DisplayNodes/VariantManagerVariantSetNode.cpp b/Engine/Plugins/Enterprise/VariantManager/Source/VariantManager/Private/DisplayNodes/VariantManagerVariantSetNode.cpp
--- a/Engine/Plugins/Enterprise/VariantManager/Source/VariantManager/Private/DisplayNodes/VariantManagerVariantSetNode.cpp
+++ b/Engine/Plugins/Enterprise/VariantManager/Source/VariantManager/Private/DisplayNodes/VariantManagerVariantSetNode.cpp
@@ -87,10 +87,9 @@ TSharedRef<SWidget> FVariantManagerVariantSetNode::GetCustomOutlinerContent(TSha
 	return
 	SNew( SBorder )
 	.BorderImage(nullptr)
-	.Padding( FMargin( 0.0f, bIsFirstRowOfTree ? 0.0f : FVariantManagerStyle::Get().GetFloat( "VariantManager.Spacings.BorderThickness" ), 0.0f, 0.0f ) )
 	[
 		SNew(SBox)
-		.HeightOverride(78)
+		.HeightOverride(48)
 		[
 			SNew(SBorder)
 			.BorderImage(nullptr)
@@ -120,7 +119,7 @@ TSharedRef<SWidget> FVariantManagerVariantSetNode::GetCustomOutlinerContent(TSha
 				+ SHorizontalBox::Slot()
 				.VAlign(VAlign_Center)
 				.HAlign(HAlign_Left)
-				.Padding(FMargin(0.f, 0.f, 4.f, 0.f))
+				.Padding(FMargin(4.f, 0.f, 4.f, 0.f))
 				.FillWidth(1.0f)
 				[
 					EditableLabel.ToSharedRef()
@@ -192,8 +191,8 @@ TSharedRef<SWidget> FVariantManagerVariantSetNode::GetThumbnailWidget()
 	}
 
 	return SNew(SBox)
-		.WidthOverride(64)
-		.HeightOverride(64)
+		.WidthOverride(40)
+		.HeightOverride(40)
 		[
 			ThumbnailWidget.IsValid() ? ThumbnailWidget.ToSharedRef() : SNullWidget::NullWidget
 		];
diff --git a/Engine/Plugins/Enterprise/VariantManager/Source/VariantManager/Private/SVariantManager.cpp b/Engine/Plugins/Enterprise/VariantManager/Source/VariantManager/Private/SVariantManager.cpp
--- a/Engine/Plugins/Enterprise/VariantManager/Source/VariantManager/Private/SVariantManager.cpp
+++ b/Engine/Plugins/Enterprise/VariantManager/Source/VariantManager/Private/SVariantManager.cpp
@@ -284,7 +284,7 @@ void SVariantManager::Construct(const FArguments& InArgs, TSharedRef<FVariantMan
 				.VAlign(VAlign_Center)
 				[
 					SNew(SSearchBox)
-					.HintText(LOCTEXT("VariantManagerFilterText", "Filter variants"))
+					.HintText(LOCTEXT("VariantManagerFilterText", "Search Variant"))
 					.OnTextChanged(this, &SVariantManager::OnOutlinerSearchChanged)
 					.MinDesiredWidth(200)
 				]
@@ -1184,6 +1184,11 @@ void SVariantManager::CreateThumbnail()
 	for (UVariant* Var : Vars)
 	{
 		Var->SetThumbnailFromEditorViewport();
+
+		if (UVariantSet* VarSet = Var->GetParent())
+		{
+			UpdateVariantSetThumbnail(VarSet);
+		}
 	}
 	for (UVariantSet* VarSet : VarSets)
 	{
@@ -1236,6 +1241,11 @@ void SVariantManager::LoadThumbnail()
 	for (UVariant* Var : Vars)
 	{
 		Var->SetThumbnailFromFile(SourceImagePath);
+
+		if (UVariantSet* VarSet = Var->GetParent())
+		{
+			UpdateVariantSetThumbnail(VarSet);
+		}
 	}
 	for (UVariantSet* VarSet : VarSets)
 	{
@@ -1271,6 +1281,23 @@ void SVariantManager::ClearThumbnail()
 	}
 }
 
+void SVariantManager::UpdateVariantSetThumbnail(UVariantSet* InVariantSet)
+{
+	if (nullptr == InVariantSet || InVariantSet->GetThumbnail())
+	{
+		return; // VariantSet is null or already has a thumbnail
+	}
+
+	for (UVariant* Variant : InVariantSet->GetVariants())
+	{
+		if (UTexture2D* Thumbnail = Variant->GetThumbnail())
+		{
+			InVariantSet->SetThumbnailFromTexture(Thumbnail);
+			break;
+		}
+	}
+}
+
 bool SVariantManager::CanSwitchOnVariant()
 {
 	TSharedPtr<FVariantManager> VariantManager = VariantManagerPtr.Pin();
diff --git a/Engine/Plugins/Enterprise/VariantManager/Source/VariantManager/Private/SVariantManager.h b/Engine/Plugins/Enterprise/VariantManager/Source/VariantManager/Private/SVariantManager.h
--- a/Engine/Plugins/Enterprise/VariantManager/Source/VariantManager/Private/SVariantManager.h
+++ b/Engine/Plugins/Enterprise/VariantManager/Source/VariantManager/Private/SVariantManager.h
@@ -238,6 +238,10 @@ private:
 	TSharedRef<SWidget> GenerateRightTreeDependenciesRowContent();
 	TSharedRef<ITableRow> GenerateDependencyRow( FVariantDependencyModelPtr Dependency, const TSharedRef<STableViewBase>& OwnerTable, bool bInteractionEnabled );
 
+	// Set thumbnail of variant set, if it does not have one, to be the same as first found
+	// variant thumbnail
+	void UpdateVariantSetThumbnail(UVariantSet* InVariantSet);
+
 private:
 	TWeakPtr<FVariantManager> VariantManagerPtr;
 
