Variant Manager: update auto-capture button and change icons sources to use svg format.

diff --git a/Engine/Plugins/Enterprise/VariantManager/Source/VariantManager/Private/SVariantManager.cpp b/Engine/Plugins/Enterprise/VariantManager/Source/VariantManager/Private/SVariantManager.cpp
--- a/Engine/Plugins/Enterprise/VariantManager/Source/VariantManager/Private/SVariantManager.cpp
+++ b/Engine/Plugins/Enterprise/VariantManager/Source/VariantManager/Private/SVariantManager.cpp
@@ -296,32 +296,47 @@ void SVariantManager::Construct(const FArguments& InArgs, TSharedRef<FVariantMan
 				.Padding(FMargin(0.f, 0.f, VM_COMMON_PADDING+2.0f, 0.f))
 				[
 					SNew(SHorizontalBox)
+
 					+ SHorizontalBox::Slot()
-					.MaxWidth(VM_COMMON_HEADER_MAX_HEIGHT) // square aspect ratio
 					.AutoWidth()
 					.VAlign(VAlign_Center)
 					[
 						SNew(SCheckBox)
-						.Style( FVariantManagerStyle::Get(), "VariantManager.AutoCapture" )
-						.Padding( FMargin( 10.f, 10.f, 10.0f, 10.f ) ) // Give some space to show up the style image or else it will collapse to nothing since it has no content
+						.Style(& FAppStyle::Get().GetWidgetStyle<FCheckBoxStyle>("ToggleButtonCheckBox"))
 						.ToolTipText(LOCTEXT("AutoCaptureTooltip", "Enable or disable auto-capturing properties"))
 						.IsChecked_Lambda([&bAutoCaptureProperties = bAutoCaptureProperties]()
 						{
-							return bAutoCaptureProperties? ECheckBoxState::Checked : ECheckBoxState::Unchecked;
+							return bAutoCaptureProperties ? ECheckBoxState::Checked : ECheckBoxState::Unchecked;
 						})
 						.OnCheckStateChanged_Lambda([&bAutoCaptureProperties = bAutoCaptureProperties](const ECheckBoxState NewState)
 						{
 							bAutoCaptureProperties = NewState == ECheckBoxState::Checked;
 						})
-					]
-					
-					+ SHorizontalBox::Slot()
-					.AutoWidth()
-					.VAlign(VAlign_Center)
-					.Padding(5, 0, 3, 0)
-					[
-						SNew(STextBlock)
-						.Text(LOCTEXT("VariantManagerAutoCapture", "Auto Capture"))
+						[
+							SNew(SBox)
+							[
+								SNew(SHorizontalBox)
+								+SHorizontalBox::Slot()
+								.AutoWidth()
+								.VAlign(VAlign_Center)
+								[
+									SNew(SImage)
+									.ColorAndOpacity(FSlateColor::UseForeground())
+									.Image(FVariantManagerStyle::Get().GetBrush("VariantManager.AutoCapture.Icon"))
+								]
+
+								+SHorizontalBox::Slot()
+								.VAlign(VAlign_Center)
+								.Padding(8.f, 0.5f, 0.f, 0.f)
+								.AutoWidth()
+								[
+									SNew(STextBlock)
+									.TextStyle(&FAppStyle::Get().GetWidgetStyle< FTextBlockStyle >("SmallButtonText"))
+									.Justification(ETextJustify::Center)
+									.Text(LOCTEXT("VariantManagerAutoCapture", "Auto Capture"))
+								]
+							]
+						]
 					]
 				]
 			]
diff --git a/Engine/Plugins/Enterprise/VariantManager/Source/VariantManager/Private/VariantManagerStyle.cpp b/Engine/Plugins/Enterprise/VariantManager/Source/VariantManager/Private/VariantManagerStyle.cpp
--- a/Engine/Plugins/Enterprise/VariantManager/Source/VariantManager/Private/VariantManagerStyle.cpp
+++ b/Engine/Plugins/Enterprise/VariantManager/Source/VariantManager/Private/VariantManagerStyle.cpp
@@ -37,24 +37,16 @@ FVariantManagerStyle::FVariantManagerStyle() : FSlateStyleSet("VariantManagerEdi
 
 	Set( "VariantManager.Icon", new IMAGE_BRUSH_SVG("VariantManager", Icon16x16) );
 
-	/** Style used for the auto-capture button */
-	const FCheckBoxStyle AutoCaptureButtonStyle = FCheckBoxStyle()
-		.SetCheckBoxType(ESlateCheckBoxType::ToggleButton)
-		.SetUncheckedImage(IMAGE_BRUSH("AutoCapture", Icon20x20))
-		.SetUncheckedHoveredImage(IMAGE_BRUSH("AutoCapture", Icon20x20, FLinearColor(0.5f, 0.5f, 0.5f, 1.0f)))
-		.SetUncheckedPressedImage(IMAGE_BRUSH("AutoCapture", Icon20x20, FLinearColor(0.5f, 0.5f, 0.5f, 1.0f)))
-		.SetCheckedImage(IMAGE_BRUSH("AutoCaptureChecked", Icon20x20))
-		.SetCheckedHoveredImage(IMAGE_BRUSH("AutoCaptureChecked", Icon20x20, FLinearColor(0.5f, 0.5f, 0.5f, 1.0f)))
-		.SetCheckedPressedImage(IMAGE_BRUSH("AutoCaptureChecked", Icon20x20, FLinearColor(0.5f, 0.5f, 0.5f, 1.0f)));
-	Set("VariantManager.AutoCapture", AutoCaptureButtonStyle);
+	Set("VariantManager.AutoCapture.Icon", new IMAGE_BRUSH_SVG("AutoCapture", Icon20x20));
 
+	SetContentRoot(FPaths::EngineContentDir());
 	const FCheckBoxStyle RadioButtonStyle = FCheckBoxStyle()
-		.SetUncheckedImage(IMAGE_BRUSH("VariantUnchecked", Icon16x16))
-		.SetUncheckedHoveredImage(IMAGE_BRUSH("VariantUnchecked", Icon16x16, FLinearColor(0.5f, 0.5f, 0.5f, 1.0f)))
-		.SetUncheckedPressedImage(IMAGE_BRUSH("VariantUnchecked", Icon16x16, FLinearColor(0.5f, 0.5f, 0.5f, 1.0f)))
-		.SetCheckedImage(IMAGE_BRUSH("VariantChecked", Icon16x16))
-		.SetCheckedHoveredImage(IMAGE_BRUSH("VariantChecked", Icon16x16, FLinearColor(0.5f, 0.5f, 0.5f, 1.0f)))
-		.SetCheckedPressedImage(IMAGE_BRUSH("VariantChecked", Icon16x16, FLinearColor(0.5f, 0.5f, 0.5f, 1.0f)));
+		.SetUncheckedImage(IMAGE_BRUSH_SVG("/Slate/Starship/CoreWidgets/CheckBox/radio-off", Icon16x16))
+		.SetUncheckedHoveredImage(IMAGE_BRUSH_SVG("/Slate/Starship/CoreWidgets/CheckBox/radio-off", Icon16x16, FLinearColor(0.5f, 0.5f, 0.5f, 1.0f)))
+		.SetUncheckedPressedImage(IMAGE_BRUSH_SVG("/Slate/Starship/CoreWidgets/CheckBox/radio-off", Icon16x16, FLinearColor(0.5f, 0.5f, 0.5f, 1.0f)))
+		.SetCheckedImage(IMAGE_BRUSH_SVG("/Slate/Starship/CoreWidgets/CheckBox/radio-on", Icon16x16))
+		.SetCheckedHoveredImage(IMAGE_BRUSH_SVG("/Slate/Starship/CoreWidgets/CheckBox/radio-on", Icon16x16, FLinearColor(0.5f, 0.5f, 0.5f, 1.0f)))
+		.SetCheckedPressedImage(IMAGE_BRUSH_SVG("/Slate/Starship/CoreWidgets/CheckBox/radio-on", Icon16x16, FLinearColor(0.5f, 0.5f, 0.5f, 1.0f)));
 	Set("VariantManager.VariantRadioButton", RadioButtonStyle);
 
 	FSlateStyleRegistry::RegisterSlateStyle(*this);
