Fixed Horde issue 190667: Compile errors in PropertyTemplateObject.h

diff --git a/Engine/Plugins/Enterprise/VariantManager/Source/VariantManager/Private/PropertyTemplateObject.h b/Engine/Plugins/Enterprise/VariantManager/Source/VariantManager/Private/PropertyTemplateObject.h
--- a/Engine/Plugins/Enterprise/VariantManager/Source/VariantManager/Private/PropertyTemplateObject.h
+++ b/Engine/Plugins/Enterprise/VariantManager/Source/VariantManager/Private/PropertyTemplateObject.h
@@ -85,7 +85,7 @@ public:
 	bool bCapturedBoolProperty;
 
 	UPROPERTY(Transient, EditAnywhere, Category="Template", meta=(ToolTip="Captured UObject property"))
-	UObject* CapturedObjectProperty;
+	TObjectPtr<UObject> CapturedObjectProperty;
 
 	//UPROPERTY(Transient, EditAnywhere, Category="Template", meta=(ToolTip="Captured Weak UObject property"))
 	//TWeakObjectPtr<UObject> CapturedWeakObjectProperty;
