Added more missing generated.h includes.

diff --git a/Engine/Plugins/Animation/ControlRig/Source/ControlRigEditor/Private/ControlRigElementDetails.h b/Engine/Plugins/Animation/ControlRig/Source/ControlRigEditor/Private/ControlRigElementDetails.h
--- a/Engine/Plugins/Animation/ControlRig/Source/ControlRigEditor/Private/ControlRigElementDetails.h
+++ b/Engine/Plugins/Animation/ControlRig/Source/ControlRigEditor/Private/ControlRigElementDetails.h
@@ -26,6 +26,8 @@
 #include "Styling/AppStyle.h"
 #include "Algo/Transform.h"
 
+#include "ControlRigElementDetails.generated.h"
+
 class IPropertyHandle;
 
 class FRigElementKeyDetails : public IPropertyTypeCustomization
diff --git a/Engine/Plugins/Editor/DataValidation/Source/DataValidation/Public/DataValidationModule.h b/Engine/Plugins/Editor/DataValidation/Source/DataValidation/Public/DataValidationModule.h
--- a/Engine/Plugins/Editor/DataValidation/Source/DataValidation/Public/DataValidationModule.h
+++ b/Engine/Plugins/Editor/DataValidation/Source/DataValidation/Public/DataValidationModule.h
@@ -8,6 +8,8 @@
 #include "Modules/ModuleInterface.h"
 #include "Modules/ModuleManager.h"
 
+#include "DataValidationModule.generated.h"
+
 UENUM(BlueprintType)
 enum class EDataValidationUsecase : uint8
 {
diff --git a/Engine/Plugins/Enterprise/VariantManager/Source/VariantManager/Private/DisplayNodes/VariantManagerDisplayNode.h b/Engine/Plugins/Enterprise/VariantManager/Source/VariantManager/Private/DisplayNodes/VariantManagerDisplayNode.h
--- a/Engine/Plugins/Enterprise/VariantManager/Source/VariantManager/Private/DisplayNodes/VariantManagerDisplayNode.h
+++ b/Engine/Plugins/Enterprise/VariantManager/Source/VariantManager/Private/DisplayNodes/VariantManagerDisplayNode.h
@@ -7,6 +7,8 @@
 #include "Widgets/Text/SInlineEditableTextBlock.h"
 //#include "DisplayNodes/SVariantManagerEditableLabel.h"
 
+#include "VariantManagerDisplayNode.generated.h"
+
 enum class EItemDropZone;
 class FMenuBuilder;
 class FVariantManager;
diff --git a/Engine/Plugins/FX/Niagara/Source/NiagaraEditor/Private/NiagaraHlslTranslator.h b/Engine/Plugins/FX/Niagara/Source/NiagaraEditor/Private/NiagaraHlslTranslator.h
--- a/Engine/Plugins/FX/Niagara/Source/NiagaraEditor/Private/NiagaraHlslTranslator.h
+++ b/Engine/Plugins/FX/Niagara/Source/NiagaraEditor/Private/NiagaraHlslTranslator.h
@@ -11,6 +11,8 @@
 #include "NiagaraScriptSource.h"
 #include "NiagaraSimulationStageCompileData.h"
 
+#include "NiagaraHlslTranslator.generated.h"
+
 class Error;
 class UNiagaraGraph;
 class UNiagaraNode;
diff --git a/Engine/Plugins/Runtime/WaveTable/Source/WaveTableEditor/Public/WaveTableCurveEditorViewStacked.h b/Engine/Plugins/Runtime/WaveTable/Source/WaveTableEditor/Public/WaveTableCurveEditorViewStacked.h
--- a/Engine/Plugins/Runtime/WaveTable/Source/WaveTableEditor/Public/WaveTableCurveEditorViewStacked.h
+++ b/Engine/Plugins/Runtime/WaveTable/Source/WaveTableEditor/Public/WaveTableCurveEditorViewStacked.h
@@ -12,6 +12,7 @@
 #include "Views/SCurveEditorViewStacked.h"
 #include "Views/SInteractiveCurveEditorView.h"
 
+#include "WaveTableCurveEditorViewStacked.generated.h"
 
 // Forward Declarations
 class UCurveFloat;
diff --git a/Engine/Source/Developer/MaterialUtilities/Public/MaterialUtilities.h b/Engine/Source/Developer/MaterialUtilities/Public/MaterialUtilities.h
--- a/Engine/Source/Developer/MaterialUtilities/Public/MaterialUtilities.h
+++ b/Engine/Source/Developer/MaterialUtilities/Public/MaterialUtilities.h
@@ -15,6 +15,8 @@
 #include "ShadowMap.h"
 #include "ImageUtils.h"
 
+#include "MaterialUtilities.generated.h"
+
 class ALandscapeProxy;
 class Error;
 class ULandscapeComponent;
diff --git a/Engine/Source/Editor/HierarchicalLODOutliner/Private/HierarchicalLODType.h b/Engine/Source/Editor/HierarchicalLODOutliner/Private/HierarchicalLODType.h
--- a/Engine/Source/Editor/HierarchicalLODOutliner/Private/HierarchicalLODType.h
+++ b/Engine/Source/Editor/HierarchicalLODOutliner/Private/HierarchicalLODType.h
@@ -5,6 +5,8 @@
 #include "CoreMinimal.h"
 #include "UObject/ObjectMacros.h"
 
+#include "HierarchicalLODType.generated.h"
+
 UENUM()	
 enum class EHierarchicalLODActionType : uint8
 {
diff --git a/Engine/Source/Editor/Kismet/Public/BlueprintActionMenuUtils.h b/Engine/Source/Editor/Kismet/Public/BlueprintActionMenuUtils.h
--- a/Engine/Source/Editor/Kismet/Public/BlueprintActionMenuUtils.h
+++ b/Engine/Source/Editor/Kismet/Public/BlueprintActionMenuUtils.h
@@ -9,6 +9,8 @@
 #include "Templates/SharedPointer.h"
 #include "UObject/ObjectMacros.h"
 
+#include "BlueprintActionMenuUtils.generated.h"
+
 class UClass;
 class UK2Node;
 struct FBlueprintActionContext;
diff --git a/Engine/Source/Editor/Sequencer/Public/ISequencer.h b/Engine/Source/Editor/Sequencer/Public/ISequencer.h
--- a/Engine/Source/Editor/Sequencer/Public/ISequencer.h
+++ b/Engine/Source/Editor/Sequencer/Public/ISequencer.h
@@ -32,6 +32,8 @@
 #include "Widgets/Input/NumericTypeInterface.h"
 #include "Widgets/SWidget.h"
 
+#include "ISequencer.generated.h"
+
 class AActor;
 class ACameraActor;
 class FSequencerKeyCollection;
diff --git a/Engine/Source/Editor/Sequencer/Public/ISequencerSection.h b/Engine/Source/Editor/Sequencer/Public/ISequencerSection.h
--- a/Engine/Source/Editor/Sequencer/Public/ISequencerSection.h
+++ b/Engine/Source/Editor/Sequencer/Public/ISequencerSection.h
@@ -22,6 +22,8 @@
 #include "Widgets/SNullWidget.h"
 #include "Widgets/SWidget.h"
 
+#include "ISequencerSection.generated.h"
+
 class FMenuBuilder;
 class FSequencerSectionPainter;
 class IDetailsView;
diff --git a/Engine/Source/Editor/UnrealEd/Public/IPropertyAccessCompiler.h b/Engine/Source/Editor/UnrealEd/Public/IPropertyAccessCompiler.h
--- a/Engine/Source/Editor/UnrealEd/Public/IPropertyAccessCompiler.h
+++ b/Engine/Source/Editor/UnrealEd/Public/IPropertyAccessCompiler.h
@@ -7,6 +7,8 @@
 #include "UObject/ObjectMacros.h"
 #include "Features/IModularFeature.h"
 
+#include "IPropertyAccessCompiler.generated.h"
+
 class UObject;
 struct FPropertyAccessLibrary;
 
diff --git a/Engine/Source/Editor/UnrealEd/Public/Kismet2/KismetDebugUtilities.h b/Engine/Source/Editor/UnrealEd/Public/Kismet2/KismetDebugUtilities.h
--- a/Engine/Source/Editor/UnrealEd/Public/Kismet2/KismetDebugUtilities.h
+++ b/Engine/Source/Editor/UnrealEd/Public/Kismet2/KismetDebugUtilities.h
@@ -8,6 +8,8 @@
 #include "UObject/Class.h"
 #include "Engine/Selection.h"
 
+#include "KismetDebugUtilities.generated.h"
+
 static_assert(DO_BLUEPRINT_GUARD, "KismetDebugUtilities assumes BP exception tracking is enabled");
 
 class UBlueprint;
diff --git a/Engine/Source/Runtime/Engine/Public/MaterialDomain.h b/Engine/Source/Runtime/Engine/Public/MaterialDomain.h
--- a/Engine/Source/Runtime/Engine/Public/MaterialDomain.h
+++ b/Engine/Source/Runtime/Engine/Public/MaterialDomain.h
@@ -5,6 +5,8 @@
 #include "UObject/ObjectMacros.h"
 #include "Containers/UnrealString.h"
 
+#include "MaterialDomain.generated.h"
+
 /** Defines the domain of a material. */
 UENUM()
 enum EMaterialDomain : int
