[VariantManager] Static mesh property now appears in the variant manager properties list

- overlay material shows up in the list too

diff --git a/Engine/Plugins/Enterprise/VariantManager/Source/VariantManager/Private/VariantManagerPropertyCapturer.cpp b/Engine/Plugins/Enterprise/VariantManager/Source/VariantManager/Private/VariantManagerPropertyCapturer.cpp
--- a/Engine/Plugins/Enterprise/VariantManager/Source/VariantManager/Private/VariantManagerPropertyCapturer.cpp
+++ b/Engine/Plugins/Enterprise/VariantManager/Source/VariantManager/Private/VariantManagerPropertyCapturer.cpp
@@ -276,7 +276,7 @@ bool FPropertyCaptureHelper::CanCaptureProperty(const UStruct* ContainerClass, c
 						checkSlow(PropertyParam); // Fix static analysis warning
 						bool bParamCanBeHandled = false;
 
-						if (PropertyParam->GetClass() == Property->GetClass())
+						if (PropertyParam->SameType(Property))
 						{
 							if (Property->GetClass()->IsChildOf(FStructProperty::StaticClass()))
 							{
