[AutoViz] Minor updates to Variant Manager
- Added Function Arguments support
- Added VariantManager.FunctionArgumentsUI CVar and Function Arguments UI
- Added ability to manipulate function callers are arguments from API
- Added VariantManager.CapturePropertyDialog (0 or 1)
- Added CapturePropertyWidgetTwo and CaptureActorsWidgetTwo for updated layout/features
- Added VariantManager.AssetPicker to control layout of SVariantManager and allow for quick switching between Level Variant Sets (and display of the current one)
- Modified CaptureDialog to support legacy layout (verision 0) or new one (version 1)
- Version 1 Capture Dialog using SSplitter if actors and properties are to be selected
- *WidgetTwo dialogs use selection instead of checkboxes
- Added 3 filters for properties and filter bar
- Added favorites with serialization
- Added ability to create function callers
- Added ability to get director blueprint
- Add/Create function callers will not get director blueprint
- Updated Variant Set and Variant thumbnails in Variant UI
- Allowing setting Variant/Variant Set thumbnails from an object
- Fixed lack of UI highlighting on Variant Manager Tab and object property details
- Moved search to be over the Variant Set/Variant tree

#Synced-CL 51657968

diff --git a/Engine/Plugins/Enterprise/VariantManager/Source/VariantManager/Private/ConsoleVariables.cpp b/Engine/Plugins/Enterprise/VariantManager/Source/VariantManager/Private/ConsoleVariables.cpp
new file mode 100644
--- /dev/null
+++ b/Engine/Plugins/Enterprise/VariantManager/Source/VariantManager/Private/ConsoleVariables.cpp
@@ -0,0 +1,30 @@
+// Copyright Epic Games, Inc. All Rights Reserved.
+
+#include "ConsoleVariables.h"
+
+namespace UE::VariantManager
+{
+	TAutoConsoleVariable<int32> CVar_VariantManager_CapturePropertyDialog(
+		TEXT("VariantManager.CapturePropertyDialog"),
+		0,
+		TEXT("Controls which Capture Property Dialog UI is used (0=default, 1=updated with last used, two columns)"));
+
+	TAutoConsoleVariable<bool> CVar_VariantManager_FunctionArgumentsUI(
+		TEXT("VariantManager.FunctionArgumentsUI"),
+		false,
+		TEXT("Enables UI for editing function caller arguments"));
+
+	TAutoConsoleVariable<bool> CVar_VariantManager_AssetPicker(
+		TEXT("VariantManager.AssetPicker"),
+		false,
+		TEXT("Improves the layout and adds asset picker to switch between Level Variant Sets."));
+
+
+	// Controls which Capture Property Dialog UI is used (0=default, 1=updated with last used, two columns)
+	TAutoConsoleVariable<int32>& Get_CVar_VariantManager_CapturePropertyDialog() { return CVar_VariantManager_CapturePropertyDialog; }
+	// Enables UI for editing function caller arguments
+	TAutoConsoleVariable<bool>& Get_CVar_VariantManager_FunctionArgumentsUI() { return CVar_VariantManager_FunctionArgumentsUI; }
+	// Improves the layout and adds asset picker to switch between Level Variant Sets.
+	TAutoConsoleVariable<bool>& Get_CVar_VariantManager_AssetPicker() { return CVar_VariantManager_AssetPicker; }
+
+}
\ No newline at end of file
diff --git a/Engine/Plugins/Enterprise/VariantManager/Source/VariantManager/Private/ConsoleVariables.h b/Engine/Plugins/Enterprise/VariantManager/Source/VariantManager/Private/ConsoleVariables.h
new file mode 100644
--- /dev/null
+++ b/Engine/Plugins/Enterprise/VariantManager/Source/VariantManager/Private/ConsoleVariables.h
@@ -0,0 +1,16 @@
+// Copyright Epic Games, Inc. All Rights Reserved.
+
+#pragma once
+
+#include "HAL/ConsoleManager.h"
+
+namespace UE::VariantManager
+{
+	// Controls which Capture Property Dialog UI is used (0=default, 1=updated with last used, two columns)
+	TAutoConsoleVariable<int32>& Get_CVar_VariantManager_CapturePropertyDialog();
+	// Enables UI for editing function caller arguments
+	TAutoConsoleVariable<bool>& Get_CVar_VariantManager_FunctionArgumentsUI();
+	// Improves the layout and adds asset picker to switch between Level Variant Sets.
+	TAutoConsoleVariable<bool>& Get_CVar_VariantManager_AssetPicker();
+}
+
diff --git a/Engine/Plugins/Enterprise/VariantManager/Source/VariantManager/Private/Dialogs/SCaptureDialog.cpp b/Engine/Plugins/Enterprise/VariantManager/Source/VariantManager/Private/Dialogs/SCaptureDialog.cpp
--- a/Engine/Plugins/Enterprise/VariantManager/Source/VariantManager/Private/Dialogs/SCaptureDialog.cpp
+++ b/Engine/Plugins/Enterprise/VariantManager/Source/VariantManager/Private/Dialogs/SCaptureDialog.cpp
@@ -2,7 +2,10 @@
 
 #include "SCaptureDialog.h"
 
+#include "ConsoleVariables.h"
 #include "Editor.h"
+#include "SCapturedActorsWidgetTwo.h"
+#include "SCapturedPropertiesWidgetTwo.h"
 #include "Widgets/Layout/SUniformGridPanel.h"
 #include "Widgets/Input/SButton.h"
 #include "Widgets/Input/SSearchBox.h"
@@ -21,6 +24,7 @@
 #define PROPERTY_DEFAULT_WIDTH 500.0f
 #define PROPERTY_DEFAULT_HEIGHT 571.0f
 
+
 void LoadWidthHeightFromIniFile(float& OutWidth, float& OutHeight, FString WindowName)
 {
 	FString WidthKey = *(WindowName + TEXT("Width"));  // Something like 'ActorPropertyDialogWidth'
@@ -99,8 +103,14 @@ void SCaptureDialog::Construct(const FArguments& InArgs)
 	bUserAccepted = false;
 	ECaptureDialogType DialogType = InArgs._DialogType;
 	const TArray<UObject*>* ObjectsToCapture = InArgs._ObjectsToCapture;
-
-	FVariantManagerPropertyCapturer::CaptureProperties(*ObjectsToCapture, CapturedProperties);
+	if (ObjectsToCapture)
+	{
+		FVariantManagerPropertyCapturer::CaptureProperties(*ObjectsToCapture, CapturedProperties);
+	}
+	else
+	{	// No objects to display, default to property only dialog to avoid errors
+		DialogType = ECaptureDialogType::Property;
+	}
 
 	FString WindowName;
 	FString WindowTitle;
@@ -125,67 +135,160 @@ void SCaptureDialog::Construct(const FArguments& InArgs)
 		UE_LOG(LogVariantManager, Error, TEXT("Invalid SCaptureDialog type!"));
 		return;
 	}
-
+	
 	TSharedRef<SVerticalBox> MainVerticalBox = SNew(SVerticalBox);
 
-	if (DialogType == ECaptureDialogType::ActorAndProperty)
+	DialogVersion = UE::VariantManager::Get_CVar_VariantManager_CapturePropertyDialog().GetValueOnAnyThread();
+
+	if (DialogVersion == 1)
 	{
-		ActorWidget = SNew(SCapturedActorsWidget).Actors(ObjectsToCapture);
+		Splitter = SNew(SSplitter)
+			.Orientation(Orient_Vertical);
+
+		if (DialogType == ECaptureDialogType::ActorAndProperty)
+		{
+			ActorWidget = SNew(SCapturedActorsWidgetTwo).Actors(ObjectsToCapture);
+			Splitter->AddSlot()
+			        .SizeRule(SSplitter::FractionOfParent)
+			        .Value(0.2)
+			[
+				SNew(SBox)
+				.Padding(0.0f, 0.0f, 0.0f, 10.0f)
+				[
+					ActorWidget.ToSharedRef()
+				]
+			];
+		}
+
+		PropertiesWidget = SNew(SCapturedPropertiesWidgetTwo).PropertyPaths(&CapturedProperties);
+		SearchBox = SNew(SSearchBox)
+			.HintText(LOCTEXT("SCapturedPropertiesWidget", "Filter"))
+			.OnTextChanged(PropertiesWidget.ToSharedRef(), &SCapturedPropertiesWidgetBase::FilterPropertyPaths);
+
+		Splitter->AddSlot()
+		        .SizeRule(SSplitter::FractionOfParent)
+		        .Value(0.7)
+		[
+			SNew(SVerticalBox)
+
+			+ SVerticalBox::Slot()
+			.MaxHeight(20.0f)
+			.AutoHeight()
+			.VAlign(VAlign_Fill)
+			.Padding(0.0f, 0.0f, 0.0f, 4.0f)
+			[
+				SearchBox.ToSharedRef()
+			]
+
+			+ SVerticalBox::Slot()
+			.Padding(0.0f, 0.0f, 0.0f, 10.0f)
+			.AutoHeight()
+			.FillHeight(1.0f)
+			.VAlign(VAlign_Top)
+			[
+				PropertiesWidget.ToSharedRef()
+			]
+		];
+
 		MainVerticalBox->AddSlot()
-		.Padding(0.0f, 0.0f, 0.0f, 10.0f)
-		.AutoHeight()
+		               .Padding(0.0f, 0.0f, 0.0f, 10.0f)
+		               .FillHeight(1.0f)
 		[
-			ActorWidget.ToSharedRef()
+			Splitter.ToSharedRef()
 		];
 	}
+	else
+	{
+		if (DialogType == ECaptureDialogType::ActorAndProperty)
+		{
+			ActorWidget = SNew(SCapturedActorsWidget).Actors(ObjectsToCapture);
+			MainVerticalBox->AddSlot()
+			               .Padding(0.0f, 0.0f, 0.0f, 10.0f)
+			               .AutoHeight()
+			[
+				ActorWidget.ToSharedRef()
+			];
+		}
 
-	PropertiesWidget = SNew(SCapturedPropertiesWidget).PropertyPaths(&CapturedProperties);
-	SearchBox = SNew(SSearchBox)
-		.HintText(LOCTEXT("SCapturedPropertiesWidget", "Filter"))
-		.OnTextChanged(PropertiesWidget.ToSharedRef(), &SCapturedPropertiesWidget::FilterPropertyPaths);
+		PropertiesWidget = SNew(SCapturedPropertiesWidget).PropertyPaths(&CapturedProperties);
 
-	// This works but conflicts with Enter/Esc to quit/accept
-	//SetWidgetToFocusOnActivate(SearchBox);
+		SearchBox = SNew(SSearchBox)
+			.HintText(LOCTEXT("SCapturedPropertiesWidget", "Filter"))
+			.OnTextChanged(PropertiesWidget.ToSharedRef(), &SCapturedPropertiesWidgetBase::FilterPropertyPaths);
 
-	MainVerticalBox->AddSlot()
-	.MaxHeight(20.0f)
-	.AutoHeight()
-	.VAlign(VAlign_Fill)
-	.Padding(0.0f, 0.0f, 0.0f, 4.0f)
-	[
-		SearchBox.ToSharedRef()
-	];
+		// This works but conflicts with Enter/Esc to quit/accept
+		//SetWidgetToFocusOnActivate(SearchBox);
 
-	MainVerticalBox->AddSlot()
-	.Padding(0.0f, 0.0f, 0.0f, 10.0f)
-	.AutoHeight()
-	.FillHeight(1.0f)
-	.VAlign(VAlign_Top)
-	[
-		PropertiesWidget.ToSharedRef()
-	];
+		MainVerticalBox->AddSlot()
+		               .MaxHeight(20.0f)
+		               .AutoHeight()
+		               .VAlign(VAlign_Fill)
+		               .Padding(0.0f, 0.0f, 0.0f, 4.0f)
+		[
+			SearchBox.ToSharedRef()
+		];
+
+		MainVerticalBox->AddSlot()
+		               .Padding(0.0f, 0.0f, 0.0f, 10.0f)
+		               .AutoHeight()
+		               .FillHeight(1.0f)
+		               .VAlign(VAlign_Top)
+		[
+			PropertiesWidget.ToSharedRef()
+		];
+	}
 
 	MainVerticalBox->AddSlot()
-	.MaxHeight(30.0f)
-	.AutoHeight()
-	.VAlign(VAlign_Center)
-	.HAlign(HAlign_Right)
-	.Padding(0.0f, 0.0f, 0.0f, 0.0f)
+	               .MaxHeight(30.0f)
+	               .AutoHeight()
+	               .VAlign(VAlign_Center)
+	               .HAlign(HAlign_Fill)
+	               .Padding(0.0f, 0.0f, 0.0f, 0.0f)
 	[
-		SNew(SUniformGridPanel)
-		.SlotPadding(FAppStyle::GetMargin("StandardDialog.SlotPadding"))
-		+SUniformGridPanel::Slot(0,0)
+		SNew(SHorizontalBox)
+		+ SHorizontalBox::Slot()
+		.FillWidth(1.0f)
+		.Padding(2.0f)
 		[
-			SNew(SPrimaryButton)
-			.Text(NSLOCTEXT("SCapturedPropertiesWidget", "ClassPickerSelectButton", "Select"))
-			.OnClicked(this, &SCaptureDialog::OnDialogConfirmed)
+			SNew(STextBlock)
+			.Text_Lambda([this]()
+				{
+					if (ActorWidget && PropertiesWidget)
+					{
+						return FText::FormatOrdered(LOCTEXT("ActorsPropertiesFmt", "{0} Actors selected and {1} Properties selected"),
+							ActorWidget->GetCurrentCheckedActors().Num(),
+							PropertiesWidget->GetCurrentCheckedProperties().Num());
+					}
+					else if (PropertiesWidget)
+					{
+						return FText::FormatOrdered(LOCTEXT("PropertiesFmt", "{0} Properties selected"),
+							PropertiesWidget->GetCurrentCheckedProperties().Num());
+					}					
+
+					return FText::GetEmpty();
+				})
+			.Visibility(UE::VariantManager::Get_CVar_VariantManager_CapturePropertyDialog().GetValueOnAnyThread() == 0 ? EVisibility::Hidden : EVisibility::Visible)
 		]
-		+SUniformGridPanel::Slot(1,0)
+
+		+ SHorizontalBox::Slot()
+		.AutoWidth()
+		.Padding(2.0f)
 		[
-			SNew(SButton)
-			.Text(NSLOCTEXT("SCapturedPropertiesWidget", "ClassPickerCancelButton", "Cancel"))
-			.TextStyle(FAppStyle::Get(), "DialogButtonText")
-			.OnClicked(this, &SCaptureDialog::OnDialogCanceled)
+			SNew(SUniformGridPanel)
+			.SlotPadding(FAppStyle::GetMargin("StandardDialog.SlotPadding"))
+			+SUniformGridPanel::Slot(0,0)
+			[
+				SNew(SPrimaryButton)
+				.Text(NSLOCTEXT("SCapturedPropertiesWidget", "ClassPickerSelectButton", "Select"))
+				.OnClicked(this, &SCaptureDialog::OnDialogConfirmed)
+			]
+			+SUniformGridPanel::Slot(1,0)
+			[
+				SNew(SButton)
+				.Text(NSLOCTEXT("SCapturedPropertiesWidget", "ClassPickerCancelButton", "Cancel"))
+				.TextStyle(FAppStyle::Get(), "DialogButtonText")
+				.OnClicked(this, &SCaptureDialog::OnDialogCanceled)
+			]
 		]
 	];
 
@@ -260,6 +363,13 @@ TSharedPtr<SCaptureDialog> SCaptureDialog::OpenCaptureDialogAsModalWindow(ECaptu
 FReply SCaptureDialog::OnDialogConfirmed()
 {
 	bUserAccepted = true;
+
+	if (DialogVersion == 1 &&
+		PropertiesWidget.IsValid())
+	{
+		StaticCastSharedPtr<SCapturedPropertiesWidgetTwo>(PropertiesWidget)->SaveFavorites();
+	}
+
 	RequestDestroyWindow();
 
 	return FReply::Handled();
diff --git a/Engine/Plugins/Enterprise/VariantManager/Source/VariantManager/Private/Dialogs/SCaptureDialog.h b/Engine/Plugins/Enterprise/VariantManager/Source/VariantManager/Private/Dialogs/SCaptureDialog.h
--- a/Engine/Plugins/Enterprise/VariantManager/Source/VariantManager/Private/Dialogs/SCaptureDialog.h
+++ b/Engine/Plugins/Enterprise/VariantManager/Source/VariantManager/Private/Dialogs/SCaptureDialog.h
@@ -3,9 +3,10 @@
 #pragma once
 
 #include "Widgets/SWindow.h"
+#include "Widgets/Layout/SSplitter.h"
 
+class SCapturedPropertiesWidgetBase;
 class SCapturedActorsWidget;
-class SCapturedPropertiesWidget;
 class SSearchBox;
 struct FCapturableProperty;
 
@@ -41,11 +42,14 @@ private:
 	virtual FReply OnKeyDown( const FGeometry& MyGeometry, const FKeyEvent& InKeyEvent ) override;
 
 	TSharedPtr<SCapturedActorsWidget> ActorWidget;
-	TSharedPtr<SCapturedPropertiesWidget> PropertiesWidget;
+	TSharedPtr<SCapturedPropertiesWidgetBase> PropertiesWidget;
 	TSharedPtr<SSearchBox> SearchBox;
 
+	TSharedPtr<SSplitter> Splitter;
+
 	TArray<TSharedPtr<FCapturableProperty>> CapturedProperties;
 	bool bUserAccepted = false;
 
 	TWeakPtr<SWindow> WeakWindow;
+	int DialogVersion = 0;
 };
diff --git a/Engine/Plugins/Enterprise/VariantManager/Source/VariantManager/Private/Dialogs/SCapturedActorsWidget.cpp b/Engine/Plugins/Enterprise/VariantManager/Source/VariantManager/Private/Dialogs/SCapturedActorsWidget.cpp
--- a/Engine/Plugins/Enterprise/VariantManager/Source/VariantManager/Private/Dialogs/SCapturedActorsWidget.cpp
+++ b/Engine/Plugins/Enterprise/VariantManager/Source/VariantManager/Private/Dialogs/SCapturedActorsWidget.cpp
@@ -13,10 +13,16 @@
 
 void SCapturedActorsWidget::Construct(const FArguments& InArgs)
 {
-	AllActors = *InArgs._Actors;
-	for (UObject* Actor : AllActors)
+	if (InArgs._Actors)
 	{
-		ActorChecked.Add(Actor, true);
+		for (UObject* Actor : *InArgs._Actors)
+		{
+			if (Actor)
+			{
+				ActorChecked.Add(Actor, true);
+				AllActors.Add(Actor);
+			}
+		}
 	}
 
     ChildSlot
@@ -42,35 +48,37 @@ TSharedRef<ITableRow> SCapturedActorsWidget::MakeRow(UObject* Item, const TShare
 	ensure(ActorChecked.Contains(Item));
 
 	AActor* ItemAsActor = Cast<AActor>(Item);
-	ensure(ItemAsActor);
-
-	return
-	SNew(STableRow<UObject*>, OwnerTable)
-	[
-		SNew(SHorizontalBox)
-		+ SHorizontalBox::Slot()
-		.Padding(8.0f, 2.0f, 10.0f, 4.0f)
-		.MaxWidth(15.0f)
-		.AutoWidth()
-		[
-			SNew(SCheckBox)
-			.ToolTipText(LOCTEXT("SCapturedActorCheckboxTooltip", "Capture this actor to the variant"))
-			.IsChecked(ActorChecked[Item] ? ECheckBoxState::Checked : ECheckBoxState::Unchecked)
-			.OnCheckStateChanged_Lambda([this, Item](ECheckBoxState NewAutoCloseState)
-			{
-				ActorChecked[Item] = (NewAutoCloseState == ECheckBoxState::Checked);
-			})
-		]
-		+ SHorizontalBox::Slot()
-		.Padding(0.0f, 2.0f, 2.0f, 4.0f)
-		.FillWidth(1.0)
-		.VAlign(VAlign_Bottom)
-		.HAlign(HAlign_Left)
+	if (ensure(ItemAsActor))
+	{
+		return SNew(STableRow<UObject*>, OwnerTable)
 		[
-			SNew(STextBlock)
-			.Text(FText::FromString(ItemAsActor->GetActorLabel()))
-		]
-	];
+			SNew(SHorizontalBox)
+			+ SHorizontalBox::Slot()
+			.Padding(8.0f, 2.0f, 10.0f, 4.0f)
+			.MaxWidth(15.0f)
+			.AutoWidth()
+			[
+				SNew(SCheckBox)
+				.ToolTipText(LOCTEXT("SCapturedActorCheckboxTooltip", "Capture this actor to the variant"))
+				.IsChecked(ActorChecked[Item] ? ECheckBoxState::Checked : ECheckBoxState::Unchecked)
+				.OnCheckStateChanged_Lambda([this, Item](ECheckBoxState NewAutoCloseState)
+				{
+					ActorChecked[Item] = (NewAutoCloseState == ECheckBoxState::Checked);
+				})
+			]
+			+ SHorizontalBox::Slot()
+			.Padding(0.0f, 2.0f, 2.0f, 4.0f)
+			.FillWidth(1.0)
+			.VAlign(VAlign_Bottom)
+			.HAlign(HAlign_Left)
+			[
+				SNew(STextBlock)
+				.Text(FText::FromString(ItemAsActor->GetActorLabel()))
+			]
+		];
+	}
+
+	return SNew(STableRow<UObject*>, OwnerTable);
 }
 
 TArray<UObject*> SCapturedActorsWidget::GetCurrentCheckedActors()
diff --git a/Engine/Plugins/Enterprise/VariantManager/Source/VariantManager/Private/Dialogs/SCapturedActorsWidget.h b/Engine/Plugins/Enterprise/VariantManager/Source/VariantManager/Private/Dialogs/SCapturedActorsWidget.h
--- a/Engine/Plugins/Enterprise/VariantManager/Source/VariantManager/Private/Dialogs/SCapturedActorsWidget.h
+++ b/Engine/Plugins/Enterprise/VariantManager/Source/VariantManager/Private/Dialogs/SCapturedActorsWidget.h
@@ -24,7 +24,7 @@ public:
 	void Construct(const FArguments& InArgs);
 	TArray<UObject*> GetCurrentCheckedActors();
 
-private:
+protected: // Changed from "private:" to allow SCapturedActorsWidgetTwo to inherit
 
 	TSharedRef<ITableRow> MakeRow(UObject* Item, const TSharedRef<STableViewBase>& OwnerTable);
 
diff --git a/Engine/Plugins/Enterprise/VariantManager/Source/VariantManager/Private/Dialogs/SCapturedActorsWidgetTwo.cpp b/Engine/Plugins/Enterprise/VariantManager/Source/VariantManager/Private/Dialogs/SCapturedActorsWidgetTwo.cpp
new file mode 100644
--- /dev/null
+++ b/Engine/Plugins/Enterprise/VariantManager/Source/VariantManager/Private/Dialogs/SCapturedActorsWidgetTwo.cpp
@@ -0,0 +1,116 @@
+// Copyright Epic Games, Inc. All Rights Reserved.
+
+
+#include "SCapturedActorsWidgetTwo.h"
+
+#include "ClassIconFinder.h"
+#include "SlateOptMacros.h"
+#include "Framework/Views/TableViewMetadata.h"
+#include "GameFramework/Actor.h"
+#include "Widgets/Layout/SExpandableArea.h"
+#include "Widgets/Input/SCheckBox.h"
+#include "Widgets/Views/SListView.h"
+
+BEGIN_SLATE_FUNCTION_BUILD_OPTIMIZATION
+
+#define LOCTEXT_NAMESPACE "SCapturedActorsWidgetTwo"
+
+void SCapturedActorsWidgetTwo::Construct(const FArguments& InArgs)
+{
+	if (InArgs._Actors)
+	{
+		for (UObject* Actor : *InArgs._Actors)
+		{
+			if (Actor)
+			{
+				ActorChecked.Add(Actor, true);
+				AllActors.Add(Actor);
+			}
+		}
+	}
+
+    ChildSlot
+	[
+		SNew(SVerticalBox)
+		
+		+ SVerticalBox::Slot()
+		.AutoHeight()
+		.Padding(0, 0, 0, 4.0f)
+		[
+			SNew(STextBlock)
+			.Text(LOCTEXT("SelectedActorsText", "Captured Actors"))
+		]
+
+		+ SVerticalBox::Slot()
+		.FillHeight(1.0f)
+		[
+			SAssignNew(ActorListView, SListView<UObject*>)
+			.SelectionMode(ESelectionMode::Multi)
+			.ListItemsSource(&AllActors)
+			.OnGenerateRow(this, &SCapturedActorsWidgetTwo::MakeRow)
+			.OnSelectionChanged(this, &SCapturedActorsWidgetTwo::OnSelectionChanged)
+			.Visibility(EVisibility::Visible)
+		]
+	];
+
+	TArray<UObject*> ActorsToSelect;
+	for (const TPair<UObject*, bool>& Pair : ActorChecked)
+	{
+		if (Pair.Value)
+		{
+			ActorsToSelect.Add(Pair.Key);
+		}
+	}
+
+	if (ActorsToSelect.Num() > 0)
+	{
+		ActorListView->SetItemSelection(ActorsToSelect, true, ESelectInfo::OnNavigation);
+	}
+}
+
+TSharedRef<ITableRow> SCapturedActorsWidgetTwo::MakeRow(UObject* Item, const TSharedRef<STableViewBase>& OwnerTable) const
+{
+	if (ensure(ActorChecked.Contains(Item)))
+	{
+		AActor* ItemAsActor = Cast<AActor>(Item);
+		if (ensure(ItemAsActor))
+		{
+			return SNew(STableRow<UObject*>, OwnerTable)
+			[
+				SNew(SHorizontalBox)
+				+ SHorizontalBox::Slot()
+				.Padding(8.0f, 2.0f, 10.0f, 4.0f)
+				.MaxWidth(15.0f)
+				.AutoWidth()
+				[
+					SNew(SImage)
+					.Image(FClassIconFinder::FindIconForActor(ItemAsActor))
+				]
+				+ SHorizontalBox::Slot()
+				.Padding(0.0f, 2.0f, 2.0f, 4.0f)
+				.FillWidth(1.0)
+				.VAlign(VAlign_Bottom)
+				.HAlign(HAlign_Left)
+				[
+					SNew(STextBlock)
+					.Text(FText::FromString(ItemAsActor->GetActorLabel()))
+				]
+			];
+		}
+	}
+
+	return SNew(STableRow<UObject*>, OwnerTable);
+}
+
+void SCapturedActorsWidgetTwo::OnSelectionChanged(UObject* Object, ESelectInfo::Type SelectType)
+{
+	for (UObject* Item : ActorListView->GetItems())
+	{
+		const bool bSelected = ActorListView->IsItemSelected(Item);
+		ActorChecked.Emplace(Item, bSelected);
+	}
+}
+
+#undef LOCTEXT_NAMESPACE
+
+END_SLATE_FUNCTION_BUILD_OPTIMIZATION
diff --git a/Engine/Plugins/Enterprise/VariantManager/Source/VariantManager/Private/Dialogs/SCapturedActorsWidgetTwo.h b/Engine/Plugins/Enterprise/VariantManager/Source/VariantManager/Private/Dialogs/SCapturedActorsWidgetTwo.h
new file mode 100644
--- /dev/null
+++ b/Engine/Plugins/Enterprise/VariantManager/Source/VariantManager/Private/Dialogs/SCapturedActorsWidgetTwo.h
@@ -0,0 +1,28 @@
+// Copyright Epic Games, Inc. All Rights Reserved.
+
+#pragma once
+
+#include "CoreMinimal.h"
+#include "SCapturedActorsWidget.h"
+#include "Widgets/Views/SListView.h"
+
+class ITableRow;
+class STableViewBase;
+struct FCapturableProperty;
+
+class SCapturedActorsWidgetTwo : public SCapturedActorsWidget
+{
+public:
+	SLATE_BEGIN_ARGS(SCapturedActorsWidgetTwo) { }
+		SLATE_ARGUMENT(const TArray<UObject*>*, Actors)
+	SLATE_END_ARGS()
+
+	/** Constructs this widget with InArgs */
+	void Construct(const FArguments& InArgs);
+	
+private:
+	TSharedRef<ITableRow> MakeRow(UObject* Item, const TSharedRef<STableViewBase>& OwnerTable) const;
+	void OnSelectionChanged(UObject* Object, ESelectInfo::Type SelectType);
+
+	TSharedPtr<SListView<UObject*>> ActorListView;
+};
diff --git a/Engine/Plugins/Enterprise/VariantManager/Source/VariantManager/Private/Dialogs/SCapturedPropertiesWidget.h b/Engine/Plugins/Enterprise/VariantManager/Source/VariantManager/Private/Dialogs/SCapturedPropertiesWidget.h
--- a/Engine/Plugins/Enterprise/VariantManager/Source/VariantManager/Private/Dialogs/SCapturedPropertiesWidget.h
+++ b/Engine/Plugins/Enterprise/VariantManager/Source/VariantManager/Private/Dialogs/SCapturedPropertiesWidget.h
@@ -10,7 +10,14 @@ class ITableRow;
 class STableViewBase;
 struct FCapturableProperty;
 
-class SCapturedPropertiesWidget : public SCompoundWidget
+class SCapturedPropertiesWidgetBase : public SCompoundWidget
+{
+public:
+	virtual TArray<TSharedPtr<FCapturableProperty>> GetCurrentCheckedProperties() = 0;
+	virtual void FilterPropertyPaths(const FText& Filter) = 0;
+};
+
+class SCapturedPropertiesWidget : public SCapturedPropertiesWidgetBase
 {
 public:
 	SLATE_BEGIN_ARGS(SCapturedPropertiesWidget) {}
@@ -24,8 +31,8 @@ public:
 	~SCapturedPropertiesWidget() = default;
 
 	void Construct(const FArguments& InArgs);
-	TArray<TSharedPtr<FCapturableProperty>> GetCurrentCheckedProperties();
-	void FilterPropertyPaths(const FText& Filter);
+	virtual TArray<TSharedPtr<FCapturableProperty>> GetCurrentCheckedProperties() override;
+	virtual void FilterPropertyPaths(const FText& Filter) override;
 
 private:
 
diff --git a/Engine/Plugins/Enterprise/VariantManager/Source/VariantManager/Private/Dialogs/SCapturedPropertiesWidgetTwo.cpp b/Engine/Plugins/Enterprise/VariantManager/Source/VariantManager/Private/Dialogs/SCapturedPropertiesWidgetTwo.cpp
new file mode 100644
--- /dev/null
+++ b/Engine/Plugins/Enterprise/VariantManager/Source/VariantManager/Private/Dialogs/SCapturedPropertiesWidgetTwo.cpp
@@ -0,0 +1,680 @@
+// Copyright Epic Games, Inc. All Rights Reserved.
+
+
+#include "SCapturedPropertiesWidgetTwo.h"
+
+#include "SlateOptMacros.h"
+
+#include "CapturableProperty.h"
+#include "VariantManagerStyle.h"
+#include "Misc/ConfigCacheIni.h"
+#include "Widgets/Layout/SExpandableArea.h"
+#include "Widgets/Views/SListView.h"
+
+#define LOCTEXT_NAMESPACE "SCapturedPropertiesWidgetTwo"
+
+BEGIN_SLATE_FUNCTION_BUILD_OPTIMIZATION
+
+void LoadFavoritesFromIniFile(TSet<FName>& OutFavoriteProperties)
+{
+	OutFavoriteProperties.Empty();
+
+	TArray<FString> Properties;
+	GConfig->GetSingleLineArray(TEXT("VariantManager"), TEXT("FavoriteProperties"), Properties, GEditorPerProjectIni);
+
+	Algo::Transform(Properties, OutFavoriteProperties, [](const FString& Value) { return FName(*Value); });
+}
+
+void SaveFavoritesToIniFile(const TSet<FName>& FavoriteProperties)
+{
+	TArray<FString> Properties;
+	Algo::Transform(FavoriteProperties, Properties, [](const FName& Name) { return Name.ToString(); });
+
+	GConfig->SetSingleLineArray(TEXT("VariantManager"), TEXT("FavoriteProperties"), Properties, GEditorPerProjectIni);
+	GConfig->Flush(false, GEditorPerProjectIni);
+}
+
+static FName ColumnFavorite = TEXT("Favorite");
+static FName ColumnPath = TEXT("Path");
+static FName ColumnProperty = TEXT("Property");
+
+SCapturedPropertiesWidgetTwo::~SCapturedPropertiesWidgetTwo()
+{
+	SaveFavorites();
+}
+
+DECLARE_DELEGATE_OneParam(FFavoriteChanged, bool);
+
+class SFavoriteWidget : public SImage
+{
+public:
+	SLATE_BEGIN_ARGS(SFavoriteWidget) {}
+		SLATE_EVENT(FFavoriteChanged, OnFavoriteChanged)
+	SLATE_END_ARGS()
+
+	/** Construct this widget */
+	void Construct(const FArguments& InArgs, TSharedPtr<FCapturableProperty> InTargetRow)
+	{
+		TargetProperty = InTargetRow;
+		OnFavoriteChanged = InArgs._OnFavoriteChanged;
+
+		SImage::Construct(
+			SImage::FArguments()
+			.ColorAndOpacity(this, &SFavoriteWidget::GetForegroundColor)
+			.Image(this, &SFavoriteWidget::GetBrush)
+		);
+
+		FavoriteHoveredBrush = FAppStyle::GetBrush("Icons.Star");
+		FavoriteNotHoveredBrush = FAppStyle::GetBrush("Icons.Star");
+		NotFavoriteHoveredBrush = FAppStyle::GetBrush("PropertyWindow.Favorites_Disabled");
+		NotFavoriteNotHoveredBrush = FAppStyle::GetBrush("PropertyWindow.Favorites_Disabled");
+	}
+
+private:
+	FReply HandleClick() const
+	{
+		if (!IsEnabled())
+		{
+			return FReply::Unhandled();
+		}
+
+		ToggleFavorite();
+
+		return FReply::Handled();
+	}
+
+	virtual FReply OnMouseButtonDoubleClick(const FGeometry& InMyGeometry, const FPointerEvent& InMouseEvent) override
+	{
+		return HandleClick();
+	}
+
+	virtual FReply OnMouseButtonDown(const FGeometry& MyGeometry, const FPointerEvent& MouseEvent) override
+	{
+		if (MouseEvent.GetEffectingButton() == EKeys::LeftMouseButton)
+		{
+			return HandleClick();
+		}
+
+		return FReply::Unhandled();
+	}
+
+	virtual FReply OnMouseButtonUp(const FGeometry& MyGeometry, const FPointerEvent& MouseEvent) override
+	{
+		if (MouseEvent.GetEffectingButton() == EKeys::LeftMouseButton)
+		{
+			return FReply::Handled();
+		}
+
+		return FReply::Unhandled();
+	}
+
+	const FSlateBrush* GetBrush() const
+	{
+		if (IsFavorite())
+		{
+			return IsHovered() ? FavoriteHoveredBrush : FavoriteNotHoveredBrush;
+		}
+		else
+		{
+			if (IsEnabled())
+			{
+				return IsHovered() ? NotFavoriteHoveredBrush : NotFavoriteNotHoveredBrush;
+			}
+		}
+
+		return FStyleDefaults::GetNoBrush();
+	}
+
+	virtual FSlateColor GetForegroundColor() const override
+	{
+		if (!IsFavorite() && !IsEnabled())
+		{
+			return FLinearColor::Transparent;
+		}
+
+		if (IsHovered())
+		{
+			return FAppStyle::Get().GetSlateColor("Colors.ForegroundHover");
+		}
+
+		return FSlateColor::UseForeground();
+	}
+
+	bool IsFavorite() const
+	{
+		return TargetProperty ? TargetProperty->Favorite : false;
+	}
+
+	void ToggleFavorite() const
+	{
+		if (TargetProperty)
+		{
+			OnFavoriteChanged.ExecuteIfBound(!TargetProperty->Favorite);
+		}
+	}
+
+private:
+	TSharedPtr<FCapturableProperty> TargetProperty;
+	FFavoriteChanged OnFavoriteChanged;
+
+	/** Favorite brushes for the various states */
+	const FSlateBrush* FavoriteHoveredBrush = nullptr;
+	const FSlateBrush* FavoriteNotHoveredBrush = nullptr;
+	const FSlateBrush* NotFavoriteHoveredBrush = nullptr;
+	const FSlateBrush* NotFavoriteNotHoveredBrush = nullptr;
+};
+
+void SCapturedPropertiesWidgetTwo::Construct(const FArguments& InArgs)
+{
+	if (InArgs._PropertyPaths)
+	{
+		CapturedProperties = *InArgs._PropertyPaths;
+	}
+
+	FilteredCapturedProperties = CapturedProperties;
+
+	LoadAndMarkFavorites();
+
+	AddFilters();
+
+    ChildSlot
+	[
+		SNew(SVerticalBox)
+
+		+ SVerticalBox::Slot()
+		.Padding(2.f)
+		.AutoHeight()
+		[
+			// Filter bar:
+			SAssignNew(FilterBar, SBasicFilterBar<TSharedPtr<FCapturableProperty>>)
+			.CustomFilters(Filters)
+			.FilterPillStyle(EFilterPillStyle::Basic)
+			.bPinAllFrontendFilters(true)
+			.bSortFilters(false)
+			.OnFilterChanged_Raw(this, &SCapturedPropertiesWidgetTwo::OnFilterChanged)
+		]
+
+		+ SVerticalBox::Slot()
+		.Padding(2.0f)
+		.AutoHeight()
+		[
+			SNew(STextBlock)
+			.Text(LOCTEXT("CapturedPropertiesText", "Captured Properties"))
+		]
+
+		+ SVerticalBox::Slot()
+		.FillHeight(1.0f)
+		[
+			SAssignNew(PropListView, SListView<TSharedPtr<FCapturableProperty>>)
+			.SelectionMode(ESelectionMode::Multi)
+			.ListItemsSource(&FilteredCapturedProperties)
+			.OnGenerateRow(this, &SCapturedPropertiesWidgetTwo::MakeCapturedPropertyWidget)
+			.Visibility(EVisibility::Visible)
+			.OnSelectionChanged_Lambda([this](const TSharedPtr<FCapturableProperty>& Item, ESelectInfo::Type Type)
+				{
+					for (TSharedPtr<FCapturableProperty>& Prop : CapturedProperties)
+					{
+						Prop->Checked = PropListView->IsItemSelected(Prop);
+					}
+				})
+			.OnItemToString_Debug_Lambda([](const TSharedPtr<FCapturableProperty>& Item)
+				{
+					return Item->DisplayName;
+				})
+			.HeaderRow(
+				SNew(SHeaderRow)
+
+				+ SHeaderRow::Column(ColumnPath)
+				.DefaultLabel(LOCTEXT("PathColumnHeader", "Path"))
+				.SortMode(this, &SCapturedPropertiesWidgetTwo::GetSortMode, ColumnPath)
+				.SortPriority(this,  &SCapturedPropertiesWidgetTwo::GetSortPriority, ColumnPath)
+				.OnSort(this, &SCapturedPropertiesWidgetTwo::HandleSort)
+				.FillWidth(0.7f)
+
+				+ SHeaderRow::Column(ColumnProperty)
+				.DefaultLabel(LOCTEXT("PropertyColumnHeader", "Property"))
+				.SortMode(this, &SCapturedPropertiesWidgetTwo::GetSortMode, ColumnProperty)
+				.SortPriority(this,  &SCapturedPropertiesWidgetTwo::GetSortPriority, ColumnProperty)
+				.OnSort(this, &SCapturedPropertiesWidgetTwo::HandleSort)
+				.FillWidth(0.3f)
+
+				+ SHeaderRow::Column(ColumnFavorite)
+				.DefaultLabel(LOCTEXT("FavoriteColumnHeader", "Favorite"))
+				.ToolTipText(LOCTEXT("FavoriteColumnHeader", "Favorite"))
+				.SortMode(this, &SCapturedPropertiesWidgetTwo::GetSortMode, ColumnFavorite)
+				.SortPriority(this, &SCapturedPropertiesWidgetTwo::GetSortPriority, ColumnFavorite)
+				.OnSort(this, &SCapturedPropertiesWidgetTwo::HandleSort)
+				.HAlignHeader(HAlign_Center)
+				.VAlignHeader(VAlign_Center)
+				.VAlignCell(VAlign_Center)
+				.HAlignCell(HAlign_Center)
+				.FixedWidth(24.0f)
+				.HeaderContentPadding(FMargin(4.0f, 0, 0, 0))
+				[
+					SNew(SImage)
+					.Image(FVariantManagerStyle::Get().GetBrush("VariantManager.IconFavorite"))
+					.ColorAndOpacity(FSlateColor::UseForeground())
+					.DesiredSizeOverride(FVector2D(16, 16))
+				]
+			)
+		]
+	];
+
+	SortColumn = ColumnFavorite;
+	SortMode = EColumnSortMode::Ascending;
+	SortPriority = EColumnSortPriority::Primary;
+
+	SortItems();
+
+	PropListView->GetHeaderRow()->RefreshColumns();
+}
+
+void SCapturedPropertiesWidgetTwo::HandleSort(
+	const EColumnSortPriority::Type NewSortPriority, 
+	const FName& ColumnId,
+	const EColumnSortMode::Type NewSortMode)
+{
+	if (NewSortPriority == EColumnSortPriority::Primary)
+	{
+		SortColumn = ColumnId;
+		SortMode = NewSortMode;
+		SortPriority = NewSortPriority;
+	}
+
+	SortItems();
+
+	if (PropListView)
+	{
+		PropListView->RequestListRefresh();
+	}
+}
+
+void SCapturedPropertiesWidgetTwo::SortItems()
+{
+	auto SortFn = [this](const TSharedPtr<FCapturableProperty>& A, const TSharedPtr<FCapturableProperty>& B)
+		{
+			if (SortColumn == ColumnFavorite)
+			{
+				if (A->Favorite != B->Favorite)	// strict-weak-ordering, A=B must return false
+				{
+					const bool bRet = A->Favorite < B->Favorite;
+
+					return SortMode == EColumnSortMode::Ascending ? bRet : !bRet;
+				}
+			}
+			else if (SortColumn == ColumnPath)
+			{
+				if (A->SubPath != B->SubPath)
+				{
+					const bool bRet = A->SubPath < B->SubPath;
+
+					return SortMode == EColumnSortMode::Ascending ? bRet : !bRet;
+				}
+			}
+			else if (SortColumn == ColumnProperty)
+			{
+				if (A->PropertyName != B->PropertyName)
+				{
+					const bool bRet = A->PropertyName < B->PropertyName;
+
+					return SortMode == EColumnSortMode::Ascending ? bRet : !bRet;
+				}
+			}
+
+			return false;
+		};
+	FilteredCapturedProperties.Sort(SortFn);
+}
+
+class FFavoritePropertyItemFilter : public FFilterBase<TSharedPtr<FCapturableProperty>>
+{
+	FChangedEvent ChangedEvent;
+public:
+	FFavoritePropertyItemFilter() :
+		FFilterBase(MakeShared<FFilterCategory>(LOCTEXT("Title", "Title"), LOCTEXT("Tooltip", "Tooltip")))
+	{
+	}
+	
+	virtual void ActiveStateChanged(bool bActive) override { OnChanged().Broadcast(); }
+	virtual FLinearColor GetColor() const override { return FLinearColor::Red; }
+	virtual FText GetDisplayName() const override { return LOCTEXT("FavoritesFilterText", "Favorites"); }
+
+	virtual FName GetIconName() const override { return NAME_None; }
+	virtual FString GetName() const override { return GetDisplayName().ToString(); }
+	virtual FText GetToolTipText() const override { return FText::GetEmpty(); }
+
+	virtual bool IsInverseFilter() const override { return false; }
+	virtual void LoadSettings(const FString& IniFilename, const FString& IniSection, const FString& SettingsString) override {}
+	virtual void ModifyContextMenu(FMenuBuilder& MenuBuilder) override {}
+	virtual FChangedEvent& OnChanged() override { return ChangedEvent; }
+	virtual bool PassesFilter(TSharedPtr<FCapturableProperty> InItem) const override
+	{
+		return InItem && InItem->Favorite;
+	}
+	virtual void SaveSettings(const FString& IniFilename, const FString& IniSection, const FString& SettingsString) const override {}
+};
+
+class FActorPropertyItemFilter : public FFilterBase<TSharedPtr<FCapturableProperty>>
+{
+	FChangedEvent ChangedEvent;
+public:
+	FActorPropertyItemFilter() :
+		FFilterBase(MakeShared<FFilterCategory>(LOCTEXT("Title", "Title"), LOCTEXT("Tooltip", "Tooltip")))
+	{
+	}
+	
+	virtual void ActiveStateChanged(bool bActive) override { OnChanged().Broadcast(); }
+	virtual FLinearColor GetColor() const override { return FLinearColor::Red; }
+	virtual FText GetDisplayName() const override { return LOCTEXT("ActorFilterText", "Actor"); }
+
+	virtual FName GetIconName() const override { return NAME_None; }
+	virtual FString GetName() const override { return GetDisplayName().ToString(); }
+	virtual FText GetToolTipText() const override { return FText::GetEmpty(); }
+
+	virtual bool IsInverseFilter() const override { return false; }
+	virtual void LoadSettings(const FString& IniFilename, const FString& IniSection, const FString& SettingsString) override {}
+	virtual void ModifyContextMenu(FMenuBuilder& MenuBuilder) override {}
+	virtual FChangedEvent& OnChanged() override { return ChangedEvent; }
+	virtual bool PassesFilter(TSharedPtr<FCapturableProperty> InItem) const override
+	{
+		return InItem && InItem->bIsActor;
+	}
+	virtual void SaveSettings(const FString& IniFilename, const FString& IniSection, const FString& SettingsString) const override {}
+};
+
+class FComponentPropertyItemFilter : public FFilterBase<TSharedPtr<FCapturableProperty>>
+{
+	FChangedEvent ChangedEvent;
+public:
+	FComponentPropertyItemFilter() :
+		FFilterBase(MakeShared<FFilterCategory>(LOCTEXT("Title", "Title"), LOCTEXT("Tooltip", "Tooltip")))
+	{
+	}
+	
+	virtual void ActiveStateChanged(bool bActive) override { OnChanged().Broadcast(); }
+	virtual FLinearColor GetColor() const override { return FLinearColor::Red; }
+	virtual FText GetDisplayName() const override { return LOCTEXT("ComponentFilterText", "Component"); }
+
+	virtual FName GetIconName() const override { return NAME_None; }
+	virtual FString GetName() const override { return GetDisplayName().ToString(); }
+	virtual FText GetToolTipText() const override { return FText::GetEmpty(); }
+
+	virtual bool IsInverseFilter() const override { return false; }
+	virtual void LoadSettings(const FString& IniFilename, const FString& IniSection, const FString& SettingsString) override {}
+	virtual void ModifyContextMenu(FMenuBuilder& MenuBuilder) override {}
+	virtual FChangedEvent& OnChanged() override { return ChangedEvent; }
+	virtual bool PassesFilter(TSharedPtr<FCapturableProperty> InItem) const override
+	{
+		return InItem && !InItem->bIsActor;
+	}
+	virtual void SaveSettings(const FString& IniFilename, const FString& IniSection, const FString& SettingsString) const override {}
+};
+
+void SCapturedPropertiesWidgetTwo::AddFilters()
+{
+	Filters.Add(MakeShared<FFavoritePropertyItemFilter>());
+	Filters.Add(MakeShared<FActorPropertyItemFilter>());
+	Filters.Add(MakeShared<FComponentPropertyItemFilter>());
+}
+
+void SCapturedPropertiesWidgetTwo::LoadAndMarkFavorites()
+{
+	LoadFavoritesFromIniFile(FavoriteProperties);
+
+	for (TSharedPtr<FCapturableProperty>& CapturedProperty : CapturedProperties)
+	{
+		TWeakFieldPtr<FProperty> Property = CapturedProperty->Prop.GetLeafMostProperty().Property;
+
+		if (Property.IsValid())
+		{
+			CapturedProperty->Favorite = FavoriteProperties.Contains(Property->NamePrivate);
+		}
+	}
+}
+
+void SCapturedPropertiesWidgetTwo::SetFavorite(const TSharedPtr<FCapturableProperty>& InProperty, bool bFavorite)
+{
+	if (InProperty && InProperty->Prop.GetLeafMostProperty().Property.IsValid())
+	{
+		const FName PropName = InProperty->Prop.GetLeafMostProperty().Property->NamePrivate;
+
+		if (bFavorite)
+		{
+			FavoriteProperties.Add(PropName);
+		}
+		else
+		{
+			FavoriteProperties.Remove(PropName);
+		}
+
+		for (TSharedPtr<FCapturableProperty>& CapturedProperty : CapturedProperties)
+		{
+			TWeakFieldPtr<FProperty> Property = CapturedProperty->Prop.GetLeafMostProperty().Property;
+			if (Property.IsValid())
+			{
+				CapturedProperty->Favorite = FavoriteProperties.Contains(Property->NamePrivate);
+			}
+		}
+	}
+}
+
+void SCapturedPropertiesWidgetTwo::SaveFavorites() const
+{
+	SaveFavoritesToIniFile(FavoriteProperties);
+}
+
+class SCapturedPropertiesMultiColumnTableRow : public SMultiColumnTableRow<TSharedPtr<FCapturableProperty>>
+{
+	SLATE_BEGIN_ARGS(SCapturedPropertiesMultiColumnTableRow) { }
+		SLATE_ATTRIBUTE(FText, HighlightText)
+		SLATE_ARGUMENT(TSharedPtr<FCapturableProperty>, ViewModel)
+	SLATE_END_ARGS()
+
+	SCapturedPropertiesWidgetTwo* Parent = nullptr;
+
+public:
+	void Construct(const FArguments& InArgs, const TSharedRef<STableViewBase>& InOwnerTableView, SCapturedPropertiesWidgetTwo* InParent)
+	{
+		check(InArgs._ViewModel);
+
+		HighlightText = InArgs._HighlightText;
+		ViewModel = InArgs._ViewModel;
+
+		Parent = InParent;
+
+		SMultiColumnTableRow::Construct(
+			FSuperRowType::FArguments()
+			.Style(&FAppStyle::Get().GetWidgetStyle<FTableRowStyle>("SceneOutliner.TableViewRow")),
+			InOwnerTableView);
+	}
+
+	virtual TSharedRef<SWidget> GenerateWidgetForColumn(const FName& InColumnName) override
+	{
+		if (InColumnName == ColumnFavorite)
+		{
+			return SNew(SBox)
+				.Padding(FMargin(4.0f, 4.0f))
+				[
+					SNew(SFavoriteWidget, ViewModel)
+					.IsEnabled_Lambda([this]()
+					{
+						if (IsSelected())
+						{
+							return true;
+						}
+
+						if (IsHovered() || IsDirectlyHovered())
+						{
+							return true;
+						}
+
+						return false;
+					})
+					.OnFavoriteChanged_Lambda([this](bool bFavorite)
+					{
+						if (Parent)
+						{
+							Parent->SetFavorite(ViewModel, bFavorite);
+						}
+					})
+				];
+		}
+		else if (InColumnName == ColumnPath)
+		{
+			return SNew(SBox)
+				.Padding(FMargin(4.0f, 4.0f))
+				[
+					SNew(SHorizontalBox)
+					.Visibility(EVisibility::HitTestInvisible)
+
+					+ SHorizontalBox::Slot()
+					.Padding(8.0f, 2.0f, 10.0f, 4.0f)
+					.MaxWidth(15.0f)
+					.AutoWidth()
+					[
+						SNew(SImage)
+						.Image(FAppStyle::Get().GetBrush(ViewModel->bIsActor ? "ClassIcon.Actor" : "ClassIcon.ActorComponent"))
+					]
+
+					+ SHorizontalBox::Slot()
+					[
+						SNew(STextBlock)
+						.Text(FText::FromString(ViewModel->SubPath))
+						.HighlightText(HighlightText)
+					]
+				];
+		}
+		else if (InColumnName == ColumnProperty)
+		{
+			return SNew(SBox)
+				.Padding(FMargin(4.0f, 4.0f))
+				.Visibility(EVisibility::HitTestInvisible)
+				[
+					SNew(STextBlock)
+					.Text(FText::FromString(ViewModel->PropertyName))
+					.HighlightText(HighlightText)
+				];
+		}
+
+		return SNullWidget::NullWidget;
+	}
+
+private:
+	TAttribute<FText> HighlightText;
+	TSharedPtr<FCapturableProperty> ViewModel;
+};
+
+TSharedRef<ITableRow> SCapturedPropertiesWidgetTwo::MakeCapturedPropertyWidget(
+	TSharedPtr<FCapturableProperty> Item, 
+	const TSharedRef<STableViewBase>& OwnerTable)
+{
+	return SNew(SCapturedPropertiesMultiColumnTableRow, OwnerTable, this)
+		.ViewModel(Item)
+		.HighlightText(HighlightText);
+}
+
+TArray<TSharedPtr<FCapturableProperty>> SCapturedPropertiesWidgetTwo::GetCurrentCheckedProperties()
+{
+	TArray<TSharedPtr<FCapturableProperty>> Result = CapturedProperties;
+	Result.RemoveAll([](const TSharedPtr<FCapturableProperty>& PropCapture)
+	{
+		return !PropCapture->Checked;
+	});
+
+	return Result;
+}
+
+void SCapturedPropertiesWidgetTwo::FilterPropertyPaths(const FText& Filter)
+{
+	HighlightText.Set(Filter);
+
+	OnFilterChanged();
+}
+
+void SCapturedPropertiesWidgetTwo::OnFilterChanged()
+{
+	// Reset selected state of all items:
+	for (TSharedPtr<FCapturableProperty>& PropCapture : CapturedProperties)
+	{
+		PropCapture->Checked = false;
+	}
+
+	if (PropListView)
+	{
+		PropListView->ClearSelection();
+	}
+
+	FilteredCapturedProperties.SetNumUninitialized(0, EAllowShrinking::No);
+
+	// Build a list of strings that must be matched
+	TArray<FString> FilterStrings;
+
+	FString FilterString = HighlightText.Get().ToString();
+	FilterString.TrimStartAndEndInline();
+	FilterString.ParseIntoArray(FilterStrings, TEXT(" "), true /*bCullEmpty*/);
+
+	for (TSharedPtr<FCapturableProperty>& PropCapture : CapturedProperties)
+	{
+		bool bPassedTextFilter = true;
+
+		// check each string in the filter strings list against
+		for (const FString& String : FilterStrings)
+		{
+			if (!PropCapture->DisplayName.Contains(String))
+			{
+				bPassedTextFilter = false;
+				break;
+			}
+		}
+
+		if (bPassedTextFilter)
+		{
+			bool bPassesFilterBar = false;
+			int ActiveFilters = 0;
+			for (const TSharedRef<FFilterBase<TSharedPtr<FCapturableProperty>>>& Filter : Filters)
+			{
+				if (Filter->IsActive())
+				{
+					ActiveFilters++;
+
+					if (bool bPassesFilter = Filter->PassesFilter(PropCapture))
+					{
+						bPassesFilterBar = true;
+						break;
+					}
+				}
+			}
+
+			if (ActiveFilters == 0)
+			{	// If no filters, just show:
+				bPassesFilterBar = true;
+			}
+
+			if (bPassesFilterBar)
+			{
+				FilteredCapturedProperties.Add(PropCapture);
+			}
+		}
+	}
+
+	SortItems();
+
+	if (PropListView)
+	{
+		PropListView->RebuildList();
+	}
+}
+
+EColumnSortMode::Type SCapturedPropertiesWidgetTwo::GetSortMode(const FName ColumnName) const
+{
+	return ColumnName == SortColumn ? SortMode : EColumnSortMode::None;
+}
+
+EColumnSortPriority::Type SCapturedPropertiesWidgetTwo::GetSortPriority(const FName ColumnName) const
+{
+	return ColumnName == SortColumn ? SortPriority : EColumnSortPriority::None;
+}
+
+END_SLATE_FUNCTION_BUILD_OPTIMIZATION
+
+#undef LOCTEXT_NAMESPACE
diff --git a/Engine/Plugins/Enterprise/VariantManager/Source/VariantManager/Private/Dialogs/SCapturedPropertiesWidgetTwo.h b/Engine/Plugins/Enterprise/VariantManager/Source/VariantManager/Private/Dialogs/SCapturedPropertiesWidgetTwo.h
new file mode 100644
--- /dev/null
+++ b/Engine/Plugins/Enterprise/VariantManager/Source/VariantManager/Private/Dialogs/SCapturedPropertiesWidgetTwo.h
@@ -0,0 +1,65 @@
+// Copyright Epic Games, Inc. All Rights Reserved.
+
+#pragma once
+
+#include "CoreMinimal.h"
+#include "SCapturedPropertiesWidget.h"
+#include "Filters/SBasicFilterBar.h"
+
+#include "Widgets/Views/SHeaderRow.h"
+
+template <typename ItemType> class SListView;
+
+class STableViewBase;
+class ITableRow;
+struct FCapturableProperty;
+
+class SCapturedPropertiesWidgetTwo : public SCapturedPropertiesWidgetBase
+{
+public:
+	SLATE_BEGIN_ARGS(SCapturedPropertiesWidgetTwo) { }
+		SLATE_ARGUMENT(TArray<TSharedPtr<FCapturableProperty>>*, PropertyPaths)
+	SLATE_END_ARGS()
+
+	virtual ~SCapturedPropertiesWidgetTwo() override;
+
+	/** Constructs this widget with InArgs */
+	void Construct(const FArguments& InArgs);
+
+public:
+	virtual TArray<TSharedPtr<FCapturableProperty>> GetCurrentCheckedProperties() override;
+	virtual void FilterPropertyPaths(const FText& Filter) override;
+
+	void SetFavorite(const TSharedPtr<FCapturableProperty>& Property, bool bFavorite);
+	void SaveFavorites() const;
+
+private:
+	EColumnSortMode::Type GetSortMode(const FName ColumnName) const;
+	EColumnSortPriority::Type GetSortPriority(FName ColumnName) const;
+	void HandleSort(const EColumnSortPriority::Type SortPriority, const FName& ColumnId, const EColumnSortMode::Type NewSortMode);
+	void SortItems();
+
+	void AddFilters();
+	void OnFilterChanged();
+
+	void LoadAndMarkFavorites();
+
+private:
+	TSharedRef<ITableRow> MakeCapturedPropertyWidget(TSharedPtr<FCapturableProperty> Item, const TSharedRef<STableViewBase>& OwnerTable);
+
+	TArray<TSharedPtr<FCapturableProperty>> CapturedProperties;
+
+	TSharedPtr<SListView<TSharedPtr<FCapturableProperty>>> PropListView;
+	TSharedPtr<SBasicFilterBar<TSharedPtr<FCapturableProperty>>> FilterBar;
+
+	TArray<TSharedPtr<FCapturableProperty>> FilteredCapturedProperties;
+	TArray<TSharedRef<FFilterBase<TSharedPtr<FCapturableProperty>>>> Filters;
+
+	TSet<FName> FavoriteProperties;
+
+	TAttribute<FText> HighlightText;
+
+	EColumnSortMode::Type SortMode = EColumnSortMode::Ascending;
+	EColumnSortPriority::Type SortPriority = EColumnSortPriority::Primary;
+	FName SortColumn;
+};
diff --git a/Engine/Plugins/Enterprise/VariantManager/Source/VariantManager/Private/Dialogs/SEditFunctionCallerArgumentsDialog.cpp b/Engine/Plugins/Enterprise/VariantManager/Source/VariantManager/Private/Dialogs/SEditFunctionCallerArgumentsDialog.cpp
new file mode 100644
--- /dev/null
+++ b/Engine/Plugins/Enterprise/VariantManager/Source/VariantManager/Private/Dialogs/SEditFunctionCallerArgumentsDialog.cpp
@@ -0,0 +1,126 @@
+// Copyright Epic Games, Inc. All Rights Reserved.
+
+
+#include "SEditFunctionCallerArgumentsDialog.h"
+
+#include "Editor.h"
+#include "IDetailsView.h"
+#include "PropertyEditorModule.h"
+#include "SlateOptMacros.h"
+#include "SPrimaryButton.h"
+#include "Modules/ModuleManager.h"
+#include "Widgets/Input/SButton.h"
+#include "Widgets/Layout/SSpacer.h"
+#include "Widgets/Layout/SUniformGridPanel.h"
+
+BEGIN_SLATE_FUNCTION_BUILD_OPTIMIZATION
+
+#define LOCTEXT_NAMESPACE "SEditFunctionCallerArgumentsDialog"
+
+void SEditFunctionCallerArgumentsDialog::Construct(const FArguments& InArgs, const FName& InFunctionName, const TMap<FName, FString>& InFunctionArguments)
+{
+	CallerArguments.Reset(NewObject<UEditFunctionCallerArguments>());
+	CallerArguments->FunctionName = InFunctionName;
+	CallerArguments->Arguments = InFunctionArguments;
+
+	FPropertyEditorModule& EditModule = FModuleManager::Get().GetModuleChecked<FPropertyEditorModule>("PropertyEditor");
+	FDetailsViewArgs DetailsViewArgs;
+	DetailsViewArgs.bAllowSearch = false;
+	DetailsViewArgs.NameAreaSettings = FDetailsViewArgs::HideNameArea;
+	DetailsViewArgs.bHideSelectionTip = true;
+	PropertyView = EditModule.CreateDetailView(DetailsViewArgs);
+	PropertyView->SetObject(CallerArguments.Get());
+
+	SWindow::Construct(SWindow::FArguments()
+	.Title(LOCTEXT("WindowTitle", "Edit Function Caller Arguments"))
+	.AutoCenter(EAutoCenter::PreferredWorkArea)
+	.SizingRule(ESizingRule::UserSized)
+	.ClientSize(FVector2D(500, 500))
+	.IsTopmostWindow(true)
+	[
+		SNew(SBorder)
+		.Visibility(EVisibility::Visible)
+		.BorderImage(FAppStyle::GetBrush("Brushes.Panel"))
+		.Padding(FMargin(10.0f, 10.0f, 10.0f, 10.0f))
+		[
+			SNew(SVerticalBox)
+
+			+ SVerticalBox::Slot()
+			.FillHeight(1.0f)
+			[
+				PropertyView.ToSharedRef()
+			]
+
+			+ SVerticalBox::Slot()
+			.AutoHeight()
+			.Padding(8.0f)
+			[
+				SNew(SSpacer)
+			]
+
+			+ SVerticalBox::Slot()
+			.AutoHeight()
+			.HAlign(HAlign_Right)
+			[
+				SNew(SUniformGridPanel)
+				.SlotPadding(FAppStyle::GetMargin("StandardDialog.SlotPadding"))
+				+SUniformGridPanel::Slot(0,0)
+				[
+					SNew(SPrimaryButton)
+					.Text(LOCTEXT("OkButton", "OK"))
+					.OnClicked(this, &SEditFunctionCallerArgumentsDialog::OnDialogConfirmed)
+				]
+				+SUniformGridPanel::Slot(1,0)
+				[
+					SNew(SButton)
+					.Text(LOCTEXT("CancelButton", "Cancel"))
+					.TextStyle(FAppStyle::Get(), "DialogButtonText")
+					.OnClicked(this, &SEditFunctionCallerArgumentsDialog::OnDialogCanceled)
+				]
+			]
+		]
+	]);
+}
+
+TSharedPtr<SEditFunctionCallerArgumentsDialog> 
+SEditFunctionCallerArgumentsDialog::OpenDialogAsModalWindow(
+	const FName& FunctionName,
+	const TMap<FName, FString>& Arguments)
+{
+	TSharedPtr<SEditFunctionCallerArgumentsDialog> Dialog = SNew(SEditFunctionCallerArgumentsDialog, FunctionName, Arguments);
+
+	GEditor->EditorAddModalWindow(Dialog.ToSharedRef());
+
+	return Dialog;
+
+}
+
+TMap<FName, FString> SEditFunctionCallerArgumentsDialog::GetFunctionArguments() const
+{
+	if (CallerArguments)
+	{
+		return CallerArguments->Arguments;
+	}
+
+	return {};
+}
+
+FReply SEditFunctionCallerArgumentsDialog::OnDialogConfirmed()
+{
+	bUserAccepted = true;
+	RequestDestroyWindow();
+
+	return FReply::Handled();
+}
+
+FReply SEditFunctionCallerArgumentsDialog::OnDialogCanceled()
+{
+	bUserAccepted = false;
+	RequestDestroyWindow();
+
+	return FReply::Handled();
+}
+
+#undef LOCTEXT_NAMESPACE
+
+END_SLATE_FUNCTION_BUILD_OPTIMIZATION
diff --git a/Engine/Plugins/Enterprise/VariantManager/Source/VariantManager/Private/Dialogs/SEditFunctionCallerArgumentsDialog.h b/Engine/Plugins/Enterprise/VariantManager/Source/VariantManager/Private/Dialogs/SEditFunctionCallerArgumentsDialog.h
new file mode 100644
--- /dev/null
+++ b/Engine/Plugins/Enterprise/VariantManager/Source/VariantManager/Private/Dialogs/SEditFunctionCallerArgumentsDialog.h
@@ -0,0 +1,58 @@
+// Copyright Epic Games, Inc. All Rights Reserved.
+
+#pragma once
+
+#include "CoreMinimal.h"
+#include "Widgets/SWindow.h"
+
+#include "SEditFunctionCallerArgumentsDialog.generated.h"
+
+/**
+ * Object representation for function caller arguments editing (with IDetailsView)
+ */
+UCLASS()
+class UEditFunctionCallerArguments : public UObject
+{
+	GENERATED_BODY()
+
+public:
+	UPROPERTY(VisibleAnywhere, Category="Function Caller")
+	FName FunctionName;
+
+	UPROPERTY(EditAnywhere, Category="Function Caller")
+	TMap<FName, FString> Arguments;
+};
+
+/**
+ * Dialog for editing function caller arguments
+ */
+class SEditFunctionCallerArgumentsDialog : public SWindow
+{
+public:
+	SLATE_BEGIN_ARGS(SEditFunctionCallerArgumentsDialog)
+		{
+		}
+
+	SLATE_END_ARGS()
+
+	/** Constructs this widget with InArgs */
+	void Construct(const FArguments& InArgs, const FName& InFunctionName, const TMap<FName, FString>& InFunctionArguments);
+
+public:
+	static TSharedPtr<SEditFunctionCallerArgumentsDialog> OpenDialogAsModalWindow(const FName& FunctionName, const TMap<FName, FString>& Arguments);
+
+public:
+	bool GetUserAccepted() const { return bUserAccepted; }
+
+	TMap<FName, FString> GetFunctionArguments() const;
+
+private:
+	FReply OnDialogConfirmed();
+	FReply OnDialogCanceled();
+
+private:
+	TSharedPtr<class IDetailsView> PropertyView;
+	TStrongObjectPtr<UEditFunctionCallerArguments> CallerArguments;
+
+	bool bUserAccepted = false;
+};
diff --git a/Engine/Plugins/Enterprise/VariantManager/Source/VariantManager/Private/DisplayNodes/VariantManagerActorNode.cpp b/Engine/Plugins/Enterprise/VariantManager/Source/VariantManager/Private/DisplayNodes/VariantManagerActorNode.cpp
--- a/Engine/Plugins/Enterprise/VariantManager/Source/VariantManager/Private/DisplayNodes/VariantManagerActorNode.cpp
+++ b/Engine/Plugins/Enterprise/VariantManager/Source/VariantManager/Private/DisplayNodes/VariantManagerActorNode.cpp
@@ -602,7 +602,6 @@ TSharedRef<SWidget> FVariantManagerActorNode::GetCustomOutlinerContent(TSharedPt
 					SAssignNew(EditableLabel, SInlineEditableTextBlock)
 					.IsReadOnly(this, &FVariantManagerDisplayNode::IsReadOnly)
 					.Font(FAppStyle::GetFontStyle("Sequencer.AnimationOutliner.RegularFont"))
-					.ColorAndOpacity(this, &FVariantManagerDisplayNode::GetDisplayNameColor)
 					.OnTextCommitted(this, &FVariantManagerDisplayNode::HandleNodeLabelTextChanged)
 					.Text(this, &FVariantManagerDisplayNode::GetDisplayName)
 					.ToolTipText(this, &FVariantManagerDisplayNode::GetDisplayNameToolTipText)
diff --git a/Engine/Plugins/Enterprise/VariantManager/Source/VariantManager/Private/DisplayNodes/VariantManagerFunctionPropertyNode.cpp b/Engine/Plugins/Enterprise/VariantManager/Source/VariantManager/Private/DisplayNodes/VariantManagerFunctionPropertyNode.cpp
--- a/Engine/Plugins/Enterprise/VariantManager/Source/VariantManager/Private/DisplayNodes/VariantManagerFunctionPropertyNode.cpp
+++ b/Engine/Plugins/Enterprise/VariantManager/Source/VariantManager/Private/DisplayNodes/VariantManagerFunctionPropertyNode.cpp
@@ -2,6 +2,7 @@
 
 #include "DisplayNodes/VariantManagerFunctionPropertyNode.h"
 
+#include "ConsoleVariables.h"
 #include "PropertyValue.h"
 #include "Kismet2/KismetEditorUtilities.h"
 #include "K2Node_FunctionEntry.h"
@@ -11,11 +12,16 @@
 #include "VariantObjectBinding.h"
 #include "BlueprintEditorModule.h"
 #include "FunctionCaller.h"
+#include "GraphEditorSettings.h"
 #include "VariantManagerEditorCommands.h"
 #include "Framework/MultiBox/MultiBoxBuilder.h"
 #include "ScopedTransaction.h"
+#include "VariantManagerStyle.h"
+#include "Dialogs/SEditFunctionCallerArgumentsDialog.h"
 #include "EdGraph/EdGraph.h"
+#include "Net/ReplayPlaylistTracker.h"
 #include "Widgets/Images/SImage.h"
+#include "Widgets/Images/SLayeredImage.h"
 #include "Widgets/Input/SButton.h"
 #include "Widgets/Input/SComboButton.h"
 
@@ -38,6 +44,20 @@ void FVariantManagerFunctionPropertyNode::BuildContextMenu(FMenuBuilder& MenuBui
 	MenuBuilder.BeginSection(TEXT("FunctionCaller"), LOCTEXT("FunctionCallerText", "Function Caller"));
 	MenuBuilder.AddMenuEntry(FVariantManagerEditorCommands::Get().CallFunction);
 	MenuBuilder.AddMenuEntry(FVariantManagerEditorCommands::Get().RemoveFunction);
+	
+	if (UE::VariantManager::Get_CVar_VariantManager_FunctionArgumentsUI().GetValueOnGameThread())
+	{
+		MenuBuilder.AddMenuEntry(
+			LOCTEXT("CallerArguments_Text",    "Edit Arguments..."),
+			LOCTEXT("CallerArguments_Tooltip", "Modifies arguments that are to be passed to function caller"),
+			FSlateIcon(FAppStyle::GetAppStyleSetName(), "Icon.Details"),
+			FUIAction(
+				FExecuteAction::CreateSP(this, &FVariantManagerFunctionPropertyNode::OnEditArguments),
+				FCanExecuteAction::CreateLambda([this]() { return FunctionCaller.CanHaveArguments(); })
+			)
+		);
+	}
+
 	MenuBuilder.EndSection();
 }
 
@@ -88,7 +108,9 @@ void FVariantManagerFunctionPropertyNode::CreateDirectorFunction(ULevelVariantSe
 	TSharedPtr<FVariantManager> PinnedVariantManager = GetVariantManager().Pin();
 	if (PinnedVariantManager.IsValid())
 	{
-		UK2Node_FunctionEntry* NewFunctionEntry = PinnedVariantManager->CreateDirectorFunction(InLevelVariantSets, FName(), PinClassType);
+		const bool bCreateWithArguments = UE::VariantManager::Get_CVar_VariantManager_FunctionArgumentsUI().GetValueOnAnyThread();
+
+		UK2Node_FunctionEntry* NewFunctionEntry = PinnedVariantManager->CreateDirectorFunction(InLevelVariantSets, FName(), PinClassType, bCreateWithArguments);
 		SetBindingTargetFunction(NewFunctionEntry);
 
 		FKismetEditorUtilities::BringKismetToFocusAttentionOnObject(NewFunctionEntry, false);
@@ -100,7 +122,9 @@ void FVariantManagerFunctionPropertyNode::CreateDirectorFunctionFromFunction(UFu
 	TSharedPtr<FVariantManager> PinnedVariantManager = GetVariantManager().Pin();
 	if (PinnedVariantManager.IsValid())
 	{
-		UK2Node_FunctionEntry* NewFunctionEntry = PinnedVariantManager->CreateDirectorFunctionFromFunction(QuickBindFunction, PinClassType);
+		const bool bCreateWithArguments = UE::VariantManager::Get_CVar_VariantManager_FunctionArgumentsUI().GetValueOnAnyThread();
+
+		UK2Node_FunctionEntry* NewFunctionEntry = PinnedVariantManager->CreateDirectorFunctionFromFunction(QuickBindFunction, PinClassType, bCreateWithArguments);
 		SetBindingTargetFunction(NewFunctionEntry);
 
 		FKismetEditorUtilities::BringKismetToFocusAttentionOnObject(NewFunctionEntry, false);
@@ -115,20 +139,20 @@ void FVariantManagerFunctionPropertyNode::NavigateToDefinition()
 		return;
 	}
 
-	ULevelVariantSets* LVS = PinnedVariantManager->GetCurrentLevelVariantSets();
-
-	UK2Node_FunctionEntry* EntryPoint = FunctionCaller.GetFunctionEntry();
-	if (EntryPoint)
-	{
-		FKismetEditorUtilities::BringKismetToFocusAttentionOnObject(EntryPoint, false);
-	}
-	else
+	if (ULevelVariantSets* LVS = PinnedVariantManager->GetCurrentLevelVariantSets())
 	{
-		// Copy the body of BringKismetToFocusAttentionOnObject without the JumpToHyperlink as it would throw a warning
-		TSharedPtr<IBlueprintEditor> BlueprintEditor = FKismetEditorUtilities::GetIBlueprintEditorForObject(LVS->GetDirectorGeneratedBlueprint(), true);
-		if (BlueprintEditor.IsValid())
+		if (UK2Node_FunctionEntry* EntryPoint = FunctionCaller.GetFunctionEntry())
 		{
-			BlueprintEditor->FocusWindow();
+			FKismetEditorUtilities::BringKismetToFocusAttentionOnObject(EntryPoint, false);
+		}
+		else
+		{
+			// Copy the body of BringKismetToFocusAttentionOnObject without the JumpToHyperlink as it would throw a warning
+			TSharedPtr<IBlueprintEditor> BlueprintEditor = FKismetEditorUtilities::GetIBlueprintEditorForObject(LVS->GetDirectorGeneratedBlueprint(), true);
+			if (BlueprintEditor.IsValid())
+			{
+				BlueprintEditor->FocusWindow();
+			}
 		}
 	}
 }
@@ -143,6 +167,11 @@ TSharedRef<SWidget> FVariantManagerFunctionPropertyNode::GetMenuContent()
 
 	ULevelVariantSets* LVS = PinnedVariantManager->GetCurrentLevelVariantSets();
 
+	if (LVS == nullptr)
+	{
+		return SNullWidget::NullWidget;
+	}
+
 	UVariantObjectBinding* Binding = GetObjectBinding().Get();
 	UClass* BoundObjectClass = nullptr;
 	if (Binding)
@@ -161,7 +190,7 @@ TSharedRef<SWidget> FVariantManagerFunctionPropertyNode::GetMenuContent()
 		MenuBuilder.AddMenuEntry(
 			LOCTEXT("CreateNewFunctionText",    "Create New Function"),
 			LOCTEXT("CreateNewFunctionTooltip", "Creates a new function in this FunctionDirector."),
-			FSlateIcon(FAppStyle::GetAppStyleSetName(), "Sequencer.CreateEventBinding"),
+			FSlateIcon(FVariantManagerStyle::GetAppStyleSetName(), "VariantManager.AddFunction"),
 			FUIAction(
 				FExecuteAction::CreateSP(this, &FVariantManagerFunctionPropertyNode::CreateDirectorFunction, LVS, BoundObjectClass)
 			)
@@ -172,7 +201,7 @@ TSharedRef<SWidget> FVariantManagerFunctionPropertyNode::GetMenuContent()
 			LOCTEXT("CreateQuickBinding_Tooltip", "Shows a list of functions on this object binding that can be bound directly to this event"),
 			FNewMenuDelegate::CreateSP(this, &FVariantManagerFunctionPropertyNode::PopulateQuickBindSubMenu, BoundObjectClass),
 			false /* bInOpenSubMenuOnClick */,
-			FSlateIcon(FAppStyle::GetAppStyleSetName(), "Sequencer.CreateQuickBinding"),
+			FSlateIcon(FVariantManagerStyle::GetAppStyleSetName(), "VariantManager.AddBinding"),
 			true /* bInShouldCloseWindowAfterMenuSelection */
 		);
 	}
@@ -180,7 +209,7 @@ TSharedRef<SWidget> FVariantManagerFunctionPropertyNode::GetMenuContent()
 	MenuBuilder.AddMenuEntry(
 		LOCTEXT("ClearCaller_Text",    "Clear"),
 		LOCTEXT("ClearCaller_Tooltip", "Unbinds the current function from this function caller"),
-		FSlateIcon(FAppStyle::GetAppStyleSetName(), "Sequencer.ClearEventBinding"),
+		FSlateIcon(FAppStyle::GetAppStyleSetName(), "Icons.Delete"),
 		FUIAction(
 			FExecuteAction::CreateSP(this, &FVariantManagerFunctionPropertyNode::SetBindingTargetFunction, (UK2Node_FunctionEntry*)nullptr)
 		)
@@ -228,9 +257,6 @@ TSharedRef<SWidget> FVariantManagerFunctionPropertyNode::GetMenuContent()
 							{
 								bRightTargetType = true;
 								TooltipText = ThisFunc->GetTooltipText();
-								TooltipText = FText::Format(
-									LOCTEXT("ValidInputTooltip", "Valid function with input pin that can receive\nreference to actor of class '{0}'"),
-									FText::FromString(BoundObjectClass->GetName()));
 							}
 							else
 							{
@@ -251,7 +277,14 @@ TSharedRef<SWidget> FVariantManagerFunctionPropertyNode::GetMenuContent()
 				}
 				else
 				{
-					TooltipText = LOCTEXT("InvalidTooltip", "Function must be valid and have either: \n - Zero input pins;\n - One input pin (with type matching the bound actor's class);\n - 4 input pins (with types: the bound actor's class, Level Variant Sets, Variant Set and Variant)");
+					if (UE::VariantManager::Get_CVar_VariantManager_FunctionArgumentsUI().GetValueOnAnyThread())
+					{
+						TooltipText = LOCTEXT("InvalidTooltipFunctionArguments", "Function must be valid and have either: \n - Zero input pins;\n - One input pin (with type matching the bound actor's class);\n - Two input pins (with types: the bound actor's class, arguments as map of Name to String);\n - 4 input pins (with types: the bound actor's class, Level Variant Sets, Variant Set and Variant)\n - 5 input pins (with types: the bound actor's class, Level Variant Sets, Variant Set, Variant and arguments as map of Name to String)");
+					}
+					else
+					{
+						TooltipText = LOCTEXT("InvalidTooltip", "Function must be valid and have either: \n - Zero input pins;\n - One input pin (with type matching the bound actor's class);\n - 4 input pins (with types: the bound actor's class, Level Variant Sets, Variant Set and Variant)");
+					}
 				}
 			}
 			else
@@ -260,7 +293,9 @@ TSharedRef<SWidget> FVariantManagerFunctionPropertyNode::GetMenuContent()
 			}
 
 			MenuBuilder.AddMenuEntry(
-				FText::FromString(FunctionGraph->GetName()),
+				FunctionCaller.CanHaveArguments(ThisFunc) ?
+					FText::FormatOrdered(LOCTEXT("FunctionWithArguments", "{0} (...)"), FText::FromString(FunctionGraph->GetName())) : 
+					FText::FromString(FunctionGraph->GetName()),
 				TooltipText,
 				Icon,
 				FUIAction(
@@ -384,18 +419,25 @@ bool FVariantManagerFunctionPropertyNode::PropertiesHaveDefaultValue() const
 
 TSharedPtr<SWidget> FVariantManagerFunctionPropertyNode::GetPropertyValueWidget()
 {
+	const bool bFunctionArgumentsUI = UE::VariantManager::Get_CVar_VariantManager_FunctionArgumentsUI().GetValueOnAnyThread();
+
 	// Make the button and set its foreground color to white so that it doesn't turn black when we have the
 	// row selected
-	TSharedRef<SWidget> BrowseButton = PropertyCustomizationHelpers::MakeBrowseButton(FSimpleDelegate::CreateSP(this, &FVariantManagerFunctionPropertyNode::NavigateToDefinition), LOCTEXT("NavigateToDefinition_Tip", "Navigate to this function's definition"));
+	TSharedRef<SWidget> BrowseButton = bFunctionArgumentsUI ? 
+		PropertyCustomizationHelpers::MakeEditButton(FSimpleDelegate::CreateSP(this, &FVariantManagerFunctionPropertyNode::NavigateToDefinition), LOCTEXT("EditFunction_Tip", "Edit this function")) : 
+		PropertyCustomizationHelpers::MakeBrowseButton(FSimpleDelegate::CreateSP(this, &FVariantManagerFunctionPropertyNode::NavigateToDefinition), LOCTEXT("NavigateToDefinition_Tip", "Navigate to this function's definition"));
 	TSharedRef<SButton> BrowseButtonAsButton = StaticCastSharedRef<SButton>(BrowseButton);
 	BrowseButtonAsButton->SetForegroundColor(FLinearColor(1.0f, 1.0f, 1.0f, 1.0f));
 
+	TSharedPtr<SHorizontalBox> HorizontalBox;
+
+
 	// Taken from FMovieSceneEventCustomization::CustomizeChildren
-	return SNew(SBox)
+	TSharedPtr<SWidget> Box = SNew(SBox)
 	.MinDesiredWidth(200.f)
 	.Padding(FMargin(4.0f, 3.0f, 0.0f, 2.0f))
 	[
-		SNew(SHorizontalBox)
+		SAssignNew(HorizontalBox, SHorizontalBox)
 
 		+ SHorizontalBox::Slot()
 		[
@@ -441,6 +483,100 @@ TSharedPtr<SWidget> FVariantManagerFunctionPropertyNode::GetPropertyValueWidget(
 			BrowseButtonAsButton
 		]
 	];
+
+	if (bFunctionArgumentsUI)
+	{
+		if (this->FunctionCaller.CanHaveArguments())
+		{
+			if (const UGraphEditorSettings* EditorSettings = GetDefault<UGraphEditorSettings>())
+			{
+				HorizontalBox->AddSlot()
+				.VAlign(VAlign_Center)
+				.AutoWidth()
+				.Padding(FMargin(7.0f, 0.0f, 6.0f, 0.0f))
+				[
+					SNew(SBox)
+					.HAlign(HAlign_Center)
+					.VAlign(VAlign_Center)
+					.WidthOverride(22.0f)
+					.HeightOverride(22.0f)
+					.ToolTipText(this, &FVariantManagerFunctionPropertyNode::GetFunctionArgumentsTooltipText)
+					[
+						SNew(SButton)
+						.ButtonStyle( FAppStyle::Get(), "SimpleButton" )
+						.OnClicked_Lambda([this]()
+						{
+							OnEditArguments();
+							return FReply::Handled();
+						})
+						.ContentPadding(0.0f)
+						.IsFocusable(true)
+						[ 
+							SNew(SLayeredImage, 
+								TAttribute<const FSlateBrush*>(FAppStyle::GetBrush(TEXT("Kismet.VariableList.MapValueTypeIcon"))),
+								FSlateColor::UseSubduedForeground())
+							.Image(FAppStyle::GetBrush(TEXT("Kismet.VariableList.MapKeyTypeIcon")))
+							.ColorAndOpacity(FSlateColor::UseForeground())
+						]
+					]
+				];
+			}
+		}
+	}
+
+	return Box;
+}
+
+void FVariantManagerFunctionPropertyNode::OnEditArguments()
+{
+	TSharedPtr<SEditFunctionCallerArgumentsDialog> Dialog = 
+		SEditFunctionCallerArgumentsDialog::OpenDialogAsModalWindow(
+			FunctionCaller.FunctionName,
+			FunctionCaller.FunctionArguments);
+
+	if (Dialog->GetUserAccepted())
+	{
+		const TMap<FName, FString> FunctionArguments = Dialog->GetFunctionArguments();
+
+		if (TStrongObjectPtr<UVariantObjectBinding> VariantObjectBinding = GetObjectBinding().Pin())
+		{
+			TArray<FFunctionCaller>& FunctionCallers = VariantObjectBinding->GetFunctionCallers();
+
+			if (FFunctionCaller* Found = FunctionCallers.FindByPredicate([FunctionName=FunctionCaller.FunctionName](const FFunctionCaller& Element)
+				{
+					return Element.FunctionName == FunctionName;
+				}))
+			{
+				FScopedTransaction EditFunctionArguments(
+					*FunctionCaller.FunctionName.ToString(),
+					LOCTEXT("EditFunctionCallerArgumentsTransaction", "Edit Function Caller Arguments"),
+						VariantObjectBinding.Get());
+
+				VariantObjectBinding->Modify();
+				Found->FunctionArguments = FunctionArguments;
+			}
+		}
+	}
+}
+
+FText FVariantManagerFunctionPropertyNode::GetFunctionArgumentsTooltipText() const
+{
+	if (FunctionCaller.FunctionArguments.Num() == 0)
+	{
+		return LOCTEXT("NoFunctionArguments", "No function arguments specified.");
+	}
+
+	FTextBuilder TextBuilder;
+	TextBuilder.AppendLineFormat(LOCTEXT("FunctionArguments", "{0} function arguments:"),
+		FFormatOrderedArguments({ FText::AsNumber(FunctionCaller.FunctionArguments.Num()) }));
+	
+	for (const TTuple<FName, FString>& Pair : FunctionCaller.FunctionArguments)
+	{
+		TextBuilder.AppendLineFormat(LOCTEXT("FunctionArgumentsPairFmt", "{0} = {1}"), 
+			FFormatOrderedArguments({ FText::FromName(Pair.Key), FText::FromString(Pair.Value) }));
+	}
+
+	return TextBuilder.ToText();
 }
 
 #undef LOCTEXT_NAMESPACE
diff --git a/Engine/Plugins/Enterprise/VariantManager/Source/VariantManager/Private/DisplayNodes/VariantManagerFunctionPropertyNode.h b/Engine/Plugins/Enterprise/VariantManager/Source/VariantManager/Private/DisplayNodes/VariantManagerFunctionPropertyNode.h
--- a/Engine/Plugins/Enterprise/VariantManager/Source/VariantManager/Private/DisplayNodes/VariantManagerFunctionPropertyNode.h
+++ b/Engine/Plugins/Enterprise/VariantManager/Source/VariantManager/Private/DisplayNodes/VariantManagerFunctionPropertyNode.h
@@ -82,6 +82,12 @@ protected:
 	// Generate the widget for the entire row
 	virtual TSharedPtr<SWidget> GetPropertyValueWidget() override;
 
+	// Opens Edit Arguments dialog
+	void OnEditArguments();
+
+	// Formats supplied function arguments into tooltip text
+	FText GetFunctionArgumentsTooltipText() const;
+
 	TWeakObjectPtr<UVariantObjectBinding> ObjectBinding;
 	FFunctionCaller& FunctionCaller;
 };
diff --git a/Engine/Plugins/Enterprise/VariantManager/Source/VariantManager/Private/DisplayNodes/VariantManagerPropertyNode.cpp b/Engine/Plugins/Enterprise/VariantManager/Source/VariantManager/Private/DisplayNodes/VariantManagerPropertyNode.cpp
--- a/Engine/Plugins/Enterprise/VariantManager/Source/VariantManager/Private/DisplayNodes/VariantManagerPropertyNode.cpp
+++ b/Engine/Plugins/Enterprise/VariantManager/Source/VariantManager/Private/DisplayNodes/VariantManagerPropertyNode.cpp
@@ -166,28 +166,6 @@ void FVariantManagerPropertyNode::Drop(const FDragDropEvent& DragDropEvent, EIte
 	VarManWidget->ReorderPropertyNodes(DraggedPropertyNodes, SharedThis(this), ItemDropZone);
 }
 
-// Without this, SImages (used for example for the browse and useselected buttons next to
-// object property value widgets) will have set SetColorAndOpacity(FSlateColor::UseForeground()).
-// That will cause them to automatically go black when the row is selected, which we don't want.
-// This unbinds that attribute and just places white as ColorAndOpacity for all SImages
-void RecursiveResetColorAndOpacityAttribute(TSharedPtr<SWidget> Widget)
-{
-	FString TypeString = Widget->GetTypeAsString();
-	if (TypeString == TEXT("SImage"))
-	{
-		if (TSharedPtr<SImage> WidgetAsImage = StaticCastSharedPtr<SImage>(Widget))
-		{
-			WidgetAsImage->SetColorAndOpacity(FLinearColor(1.0f, 1.0f, 1.0f));
-		}
-	}
-
-	FChildren* Children = Widget->GetChildren();
-	for (int32 Index = 0; Index < Children->Num(); ++Index)
-	{
-		RecursiveResetColorAndOpacityAttribute(Children->GetChildAt(Index));
-	}
-}
-
 // We replace the standard reset button with our own. No easy way to do this other than to dig
 // around for it in the widget hierarchy
 bool FVariantManagerPropertyNode::RecursiveDisableOldResetButton(TSharedPtr<SWidget> Widget)
@@ -221,6 +199,13 @@ TSharedPtr<SWidget> FVariantManagerPropertyNode::GetPropertyValueWidget()
 
 	// Check to see if we have all valid, equal UPropertyValues
 	UPropertyValue* FirstPropertyValue = PropertyValues[0].Get();
+
+	if (!FirstPropertyValue)
+	{
+		UE_LOG(LogVariantManager, Error, TEXT("PropertyNode has no first property value!"));
+		return SNew(SBox);
+	}
+
 	uint32 FirstPropHash = FirstPropertyValue->GetPropertyPathHash();
 	for (TWeakObjectPtr<UPropertyValue> PropertyValue : PropertyValues)
 	{
@@ -331,7 +316,6 @@ TSharedPtr<SWidget> FVariantManagerPropertyNode::GetPropertyValueWidget()
 		];
 	}
 
-	RecursiveResetColorAndOpacityAttribute(SinglePropView);
 	RecursiveDisableOldResetButton(SinglePropView);
 
 	TSharedPtr<IPropertyHandle> PropHandle = SinglePropView->GetPropertyHandle();
@@ -339,17 +323,32 @@ TSharedPtr<SWidget> FVariantManagerPropertyNode::GetPropertyValueWidget()
 	// Update widget with recorded data
 	TArray<void*> RawDataForEachObject;
 	PropHandle->AccessRawData(RawDataForEachObject);
-	void* SinglePropWidgetDataPtr = RawDataForEachObject[0]; // We'll always pass just a single object
+	void* SinglePropWidgetDataPtr = RawDataForEachObject.IsValidIndex(0) ? 
+		RawDataForEachObject[0] :	// We'll always pass just a single object
+		nullptr;					// In case nothing is passed
+
 	const TArray<uint8>& FirstRecordedData = FirstPropertyValue->GetRecordedData();
 
-	if (FSoftObjectProperty* Prop = CastField<FSoftObjectProperty>(PropHandle->GetProperty()))
-	{
-		UObject* TargetObject = *((UObject**)FirstRecordedData.GetData());
-		Prop->SetObjectPropertyValue(SinglePropWidgetDataPtr, TargetObject);
-	}
-	else
+	if (SinglePropWidgetDataPtr)
 	{
-		FMemory::Memcpy((uint8*)SinglePropWidgetDataPtr, FirstRecordedData.GetData(), FirstPropertyValue->GetValueSizeInBytes());
+		if (FSoftObjectProperty* Prop = CastField<FSoftObjectProperty>(PropHandle->GetProperty()))
+		{
+			if (FirstRecordedData.Num() >= sizeof(UObject*))
+			{
+				UObject* TargetObject = *((UObject**)(FirstRecordedData.GetData()));
+				Prop->SetObjectPropertyValue(SinglePropWidgetDataPtr, TargetObject);
+			}
+		}
+		else
+		{
+			if (FirstRecordedData.Num() >= FirstPropertyValue->GetValueSizeInBytes())
+			{
+				if (const uint8* DataPtr = FirstRecordedData.GetData())
+				{
+					FMemory::Memcpy((uint8*)SinglePropWidgetDataPtr, DataPtr, FirstPropertyValue->GetValueSizeInBytes());
+				}
+			}
+		}
 	}
 
 	// Update recorded data when user modifies the widget (modifying the widget will modify the
@@ -362,109 +361,121 @@ TSharedPtr<SWidget> FVariantManagerPropertyNode::GetPropertyValueWidget()
 TSharedRef<SWidget> FVariantManagerPropertyNode::GetCustomOutlinerContent(TSharedPtr<SVariantManagerTableRow> InTableRow)
 {
 	// Using this syncs all splitters between property nodes
-	TSharedPtr<SVariantManager> VariantManagerWidget = GetVariantManager().Pin()->GetVariantManagerWidget();
-	FVariantManagerPropertiesColumnSizeData& ColumnSizeData = VariantManagerWidget->GetPropertiesColumnSizeData();
+	if (TSharedPtr<FVariantManager> VariantManagerPtr = GetVariantManager().Pin())
+	{
+		if (TSharedPtr<SVariantManager> VariantManagerWidget = VariantManagerPtr->GetVariantManagerWidget())
+		{
+			FVariantManagerPropertiesColumnSizeData& ColumnSizeData = VariantManagerWidget->GetPropertiesColumnSizeData();
 
-	return SNew(SBox)
-	[
-		SNew(SBorder)
-		.VAlign(VAlign_Center)
-		.BorderImage(this, &FVariantManagerDisplayNode::GetNodeBorderImage)
-		.Padding(FMargin(2.0f, 0.0f, 2.0f, 0.0f))
-		[
-			SNew(SSplitter)
-			.Style(FAppStyle::Get(), "DetailsView.Splitter")
-			.PhysicalSplitterHandleSize(1.0f)
-			.HitDetectionSplitterHandleSize(5.0f)
-			.Orientation(Orient_Horizontal)
-
-			+ SSplitter::Slot()
-			.Value(ColumnSizeData.NameColumnWidth)
-			.OnSlotResized(ColumnSizeData.OnSplitterNameColumnChanged)
+			return SNew(SBox)
 			[
-				SNew(SBox)
+				SNew(SBorder)
 				.VAlign(VAlign_Center)
-				.HAlign(HAlign_Left)
-				.Padding(FMargin( 4.0, 0.0f, 0.0f, 0.0f ))
-				.HeightOverride(26) // Sum of paddings for the rows used in a details view like this
-				[
-					SAssignNew(EditableLabel, SInlineEditableTextBlock)
-					.IsReadOnly(true)
-					.Font(FAppStyle::GetFontStyle("Sequencer.AnimationOutliner.RegularFont"))
-					.ColorAndOpacity(this, &FVariantManagerDisplayNode::GetDisplayNameColor)
-					.OnTextCommitted(this, &FVariantManagerDisplayNode::HandleNodeLabelTextChanged)
-					.Text(this, &FVariantManagerDisplayNode::GetDisplayName)
-					.ToolTipText(this, &FVariantManagerDisplayNode::GetDisplayNameToolTipText)
-					.Clipping(EWidgetClipping::ClipToBounds)
-				]
-			]
-
-			+ SSplitter::Slot()
-			.Value(ColumnSizeData.ValueColumnWidth)
-			.OnSlotResized(ColumnSizeData.OnSplitterValueColumnChanged)
-			[
-				SNew(SHorizontalBox)
-				+SHorizontalBox::Slot()
+				.BorderImage(this, &FVariantManagerDisplayNode::GetNodeBorderImage)
+				.Padding(FMargin(2.0f, 0.0f, 2.0f, 0.0f))
 				[
-					SNew(SBox)
-					.Padding(FMargin( 3.0, 0.0f, 3.0f, 0.0f ))
+					SNew(SSplitter)
+					.Style(FAppStyle::Get(), "DetailsView.Splitter")
+					.PhysicalSplitterHandleSize(1.0f)
+					.HitDetectionSplitterHandleSize(5.0f)
+					.Orientation(Orient_Horizontal)
+
+					+ SSplitter::Slot()
+					.Value(ColumnSizeData.NameColumnWidth)
+					.OnSlotResized(ColumnSizeData.OnSplitterNameColumnChanged)
 					[
-						GetPropertyValueWidget().ToSharedRef()
+						SNew(SBox)
+						.VAlign(VAlign_Center)
+						.HAlign(HAlign_Left)
+						.Padding(FMargin( 4.0, 0.0f, 0.0f, 0.0f ))
+						.HeightOverride(26) // Sum of paddings for the rows used in a details view like this
+						[
+							SAssignNew(EditableLabel, SInlineEditableTextBlock)
+							.IsReadOnly(true)
+							.Font(FAppStyle::GetFontStyle("Sequencer.AnimationOutliner.RegularFont"))
+							.ColorAndOpacity(this, &FVariantManagerDisplayNode::GetDisplayNameColor)
+							.OnTextCommitted(this, &FVariantManagerDisplayNode::HandleNodeLabelTextChanged)
+							.Text(this, &FVariantManagerDisplayNode::GetDisplayName)
+							.ToolTipText(this, &FVariantManagerDisplayNode::GetDisplayNameToolTipText)
+							.Clipping(EWidgetClipping::ClipToBounds)
+						]
 					]
-				]
 
-				+SHorizontalBox::Slot()
-				.Padding( FMargin(0.0f, 0.0f, 3.0f, 0.0) )
-				.AutoWidth()
-				.VAlign( VAlign_Center )
-				[
-					SNew(SBox)
+					+ SSplitter::Slot()
+					.Value(ColumnSizeData.ValueColumnWidth)
+					.OnSlotResized(ColumnSizeData.OnSplitterValueColumnChanged)
 					[
-						SAssignNew(RecordButton, SButton)
-						.IsFocusable(false)
-						.ToolTipText(LOCTEXT("UseCurrentTooltip", "Record the current value for this property"))
-						.ButtonStyle( FAppStyle::Get(), "HoverHintOnly" )
-						.ContentPadding(6.0f)
-						.OnClicked(this, &FVariantManagerPropertyNode::RecordMultipleValues)
-						.Visibility(FVariantManagerPropertyNode::GetRecordButtonVisibility())
-						.Content()
-						[
-							SNew(STextBlock)
-							.TextStyle(FAppStyle::Get(), "ContentBrowser.TopBar.Font")
-							.Font(FAppStyle::Get().GetFontStyle("FontAwesome.11"))
-							.ShadowColorAndOpacity(FVector::ZeroVector)
-							.ShadowOffset(FVector::ZeroVector)
-							.Text(FEditorFontGlyphs::Download)
-							.Justification(ETextJustify::Center)
-							.Margin(FMargin(0.0f, 3.0f, 0.0f, 0.0f))
-						]
+						GetCustomValueContent()
 					]
 				]
+			];
+		}
+	}
 
-				+SHorizontalBox::Slot()
-				.Padding( FMargin(0.0f, 0.0f, 3.0f, 0.0) )
-				.AutoWidth()
-				.VAlign(VAlign_Center)
+	return SNullWidget::NullWidget;
+}
+
+TSharedRef<SWidget> FVariantManagerPropertyNode::GetCustomValueContent()
+{
+	return SNew(SHorizontalBox)
+		+SHorizontalBox::Slot()
+		[
+			SNew(SBox)
+			.Padding(FMargin( 3.0, 0.0f, 3.0f, 0.0f ))
+			[
+				GetPropertyValueWidget().ToSharedRef()
+			]
+		]
+
+		+SHorizontalBox::Slot()
+		.Padding( FMargin(0.0f, 0.0f, 3.0f, 0.0) )
+		.AutoWidth()
+		.VAlign( VAlign_Center )
+		[
+			SNew(SBox)
+			[
+				SAssignNew(RecordButton, SButton)
+				.IsFocusable(false)
+				.ToolTipText(LOCTEXT("UseCurrentTooltip", "Record the current value for this property"))
+				.ButtonStyle( FAppStyle::Get(), "HoverHintOnly" )
+				.ContentPadding(6.0f)
+				.OnClicked(this, &FVariantManagerPropertyNode::RecordMultipleValues)
+				.Visibility(FVariantManagerPropertyNode::GetRecordButtonVisibility())
+				.Content()
 				[
-					SNew(SBox)
-					[
-						SAssignNew(ResetButton, SButton)
-						.IsFocusable(false)
-						.ToolTipText(LOCTEXT("ResetTooltip", "Reset to the property's default value"))
-						.ButtonStyle( FAppStyle::Get(), "HoverHintOnly" )
-						.ContentPadding(FMargin(3.0f, 8.0f))
-						.OnClicked(this, &FVariantManagerPropertyNode::ResetMultipleValuesToDefault)
-						.Visibility(FVariantManagerPropertyNode::GetResetButtonVisibility())
-						.Content()
-						[
-							SNew(SImage)
-							.Image( FAppStyle::GetBrush("PropertyWindow.DiffersFromDefault") )
-						]
-					]
+					SNew(STextBlock)
+					.TextStyle(FAppStyle::Get(), "ContentBrowser.TopBar.Font")
+					.Font(FAppStyle::Get().GetFontStyle("FontAwesome.11"))
+					.ShadowColorAndOpacity(FVector::ZeroVector)
+					.ShadowOffset(FVector::ZeroVector)
+					.Text(FEditorFontGlyphs::Download)
+					.Justification(ETextJustify::Center)
+					.Margin(FMargin(0.0f, 3.0f, 0.0f, 0.0f))
 				]
 			]
 		]
-	];
+
+		+SHorizontalBox::Slot()
+		.Padding( FMargin(0.0f, 0.0f, 3.0f, 0.0) )
+		.AutoWidth()
+		.VAlign(VAlign_Center)
+		[
+			SNew(SBox)
+			[
+				SAssignNew(ResetButton, SButton)
+				.IsFocusable(false)
+				.ToolTipText(LOCTEXT("ResetTooltip", "Reset to the property's default value"))
+				.ButtonStyle( FAppStyle::Get(), "HoverHintOnly" )
+				.ContentPadding(FMargin(3.0f, 8.0f))
+				.OnClicked(this, &FVariantManagerPropertyNode::ResetMultipleValuesToDefault)
+				.Visibility(FVariantManagerPropertyNode::GetResetButtonVisibility())
+				.Content()
+				[
+					SNew(SImage)
+					.Image( FAppStyle::GetBrush("PropertyWindow.DiffersFromDefault") )
+				]
+			]
+		];
 }
 
 uint32 FVariantManagerPropertyNode::GetDisplayOrder() const
@@ -492,6 +503,53 @@ void FVariantManagerPropertyNode::SetDisplayOrder(uint32 InDisplayOrder)
 	}
 }
 
+bool FVariantManagerPropertyNode::IsActorProperty() const
+{
+	for (const TWeakObjectPtr<UPropertyValue>& PropertyValue : PropertyValues)
+	{
+		if (TStrongObjectPtr<UPropertyValue> PropValue = PropertyValue.Pin())
+		{
+			return !(PropertyValue->GetPropertyParentContainerClass() && !PropertyValue->GetPropertyParentContainerClass()->IsChildOf(AActor::StaticClass()));
+		}
+	}
+
+	return true;
+}
+
+FString FVariantManagerPropertyNode::GetSubPath() const
+{
+	FString SubPath = TEXT("Actor");
+
+	if (GetType() == EVariantManagerNodeType::Property)
+	{
+		for (const TWeakObjectPtr<UPropertyValue>& PropertyValue : PropertyValues)
+		{
+			if (TStrongObjectPtr<UPropertyValue> PropValue = PropertyValue.Pin())
+			{
+				if (PropValue->GetPropertyParentContainerClass() && !PropValue->GetPropertyParentContainerClass()->IsChildOf(AActor::StaticClass()))
+				{
+					SubPath = PropValue->GetBaseDisplayString().TrimStartAndEnd();
+				}
+			}
+		}
+	}
+
+	return SubPath;
+}
+
+FString FVariantManagerPropertyNode::GetPropertyName() const
+{
+	for (const TWeakObjectPtr<UPropertyValue>& PropertyValue : PropertyValues)
+	{
+		if (TStrongObjectPtr<UPropertyValue> PropValue = PropertyValue.Pin())
+		{
+			return PropValue->GetPropertyName().ToString();
+		}
+	}
+
+	return FString();
+}
+
 void FVariantManagerPropertyNode::UpdateRecordedDataFromSinglePropView(TSharedPtr<ISinglePropertyView> SinglePropView)
 {
 	// Warning: This also fires after UpdateSinglePropViewFromRecordedData when that fires
@@ -500,24 +558,29 @@ void FVariantManagerPropertyNode::UpdateRecordedDataFromSinglePropView(TSharedPt
 	TSharedPtr<IPropertyHandle> PropHandle = SinglePropView->GetPropertyHandle();
 	TArray<void*> RawDataForEachObject;
 	PropHandle->AccessRawData(RawDataForEachObject);
-	void* SinglePropWidgetDataPtr = RawDataForEachObject[0]; // We'll always pass just a single object
+	void* SinglePropWidgetDataPtr = RawDataForEachObject.IsValidIndex(0) ? 
+		RawDataForEachObject[0] :	// We'll always pass just a single object
+		nullptr;					// Or in case nothing is passed
 
-	for (TWeakObjectPtr<UPropertyValue> PropertyValuePtr : PropertyValues)
+	if (SinglePropWidgetDataPtr)
 	{
-		UPropertyValue* PropertyValue = PropertyValuePtr.Get();
-		if (!PropertyValue)
+		for (TWeakObjectPtr<UPropertyValue> PropertyValuePtr : PropertyValues)
 		{
-			continue;
-		}
+			UPropertyValue* PropertyValue = PropertyValuePtr.Get();
+			if (!PropertyValue)
+			{
+				continue;
+			}
 
-		if (FSoftObjectProperty* Prop = CastField<FSoftObjectProperty>(PropHandle->GetProperty()))
-	{
-			UObject* NewObj = Prop->LoadObjectPropertyValue(SinglePropWidgetDataPtr);
-			PropertyValue->SetRecordedData((uint8*)&NewObj, sizeof(UObject*));
-		}
-		else
-		{
-			PropertyValue->SetRecordedData((uint8*)SinglePropWidgetDataPtr, PropHandle->GetProperty()->GetElementSize());
+			if (FSoftObjectProperty* Prop = CastField<FSoftObjectProperty>(PropHandle->GetProperty()))
+			{
+				UObject* NewObj = Prop->LoadObjectPropertyValue(SinglePropWidgetDataPtr);
+				PropertyValue->SetRecordedData(reinterpret_cast<uint8*>(&NewObj), sizeof(UObject*));
+			}
+			else
+			{
+				PropertyValue->SetRecordedData(static_cast<uint8*>(SinglePropWidgetDataPtr), PropHandle->GetProperty()->GetElementSize());
+			}
 		}
 	}
 
@@ -556,7 +619,14 @@ FReply FVariantManagerPropertyNode::ResetMultipleValuesToDefault()
 		}
 	}
 
-	GetVariantManager().Pin()->GetVariantManagerWidget()->RefreshPropertyList();
+	if (TSharedPtr<FVariantManager> VariantManagerPtr = GetVariantManager().Pin())
+	{
+		if (TSharedPtr<SVariantManager> VariantManagerWidget = VariantManagerPtr->GetVariantManagerWidget())
+		{
+			VariantManagerWidget->RefreshPropertyList();
+		}
+	}
+
 	return FReply::Handled();
 }
 
@@ -575,7 +645,14 @@ FReply FVariantManagerPropertyNode::RecordMultipleValues()
 		}
 	}
 
-	GetVariantManager().Pin()->GetVariantManagerWidget()->RefreshPropertyList();
+	if (TSharedPtr<FVariantManager> VariantManagerPtr = GetVariantManager().Pin())
+	{
+		if (TSharedPtr<SVariantManager> VariantManagerWidgetPtr = VariantManagerPtr->GetVariantManagerWidget())
+		{
+			VariantManagerWidgetPtr->RefreshPropertyList();
+		}
+	}
+
 	return FReply::Handled();
 }
 
diff --git a/Engine/Plugins/Enterprise/VariantManager/Source/VariantManager/Private/DisplayNodes/VariantManagerPropertyNode.h b/Engine/Plugins/Enterprise/VariantManager/Source/VariantManager/Private/DisplayNodes/VariantManagerPropertyNode.h
--- a/Engine/Plugins/Enterprise/VariantManager/Source/VariantManager/Private/DisplayNodes/VariantManagerPropertyNode.h
+++ b/Engine/Plugins/Enterprise/VariantManager/Source/VariantManager/Private/DisplayNodes/VariantManagerPropertyNode.h
@@ -46,8 +46,16 @@ public:
 		return VariantManager;
 	}
 	//~ End FVariantManagerDisplayNode interface
-
+	
+	/**
+	 * Full outliner content for this property node. This includes name, splitter and value content.
+	 */
 	virtual TSharedRef<SWidget> GetCustomOutlinerContent(TSharedPtr<SVariantManagerTableRow> InTableRow) override;
+	
+	/**
+	 * Just value content for this property node.
+	 */
+	virtual TSharedRef<SWidget> GetCustomValueContent();
 
 	/**
 	 * Get the order with which the VariantManager should display this in a property list. Lower values will be shown higher up
@@ -59,6 +67,21 @@ public:
 	 */
 	virtual void SetDisplayOrder(uint32 InDisplayOrder);
 
+	/**
+	 * Flag if this is an Actor or Actor Component property.
+	 */
+	bool IsActorProperty() const;
+
+	/**
+	 * Path without the property name
+	 */
+	FString GetSubPath() const;
+
+	/**
+	 * Name of the property (last part of the path)
+	 */
+	FString GetPropertyName() const;
+
 protected:
 
 	// Callback for when user updates the property widget
diff --git a/Engine/Plugins/Enterprise/VariantManager/Source/VariantManager/Private/DisplayNodes/VariantManagerVariantNode.cpp b/Engine/Plugins/Enterprise/VariantManager/Source/VariantManager/Private/DisplayNodes/VariantManagerVariantNode.cpp
--- a/Engine/Plugins/Enterprise/VariantManager/Source/VariantManager/Private/DisplayNodes/VariantManagerVariantNode.cpp
+++ b/Engine/Plugins/Enterprise/VariantManager/Source/VariantManager/Private/DisplayNodes/VariantManagerVariantNode.cpp
@@ -15,6 +15,9 @@
 #include "VariantSet.h"
 
 #include "AssetThumbnail.h"
+#include "ConsoleVariables.h"
+#include "LevelVariantSets.h"
+#include "SAssetDropTarget.h"
 #include "Brushes/SlateImageBrush.h"
 #include "DragAndDrop/ActorDragDropGraphEdOp.h"
 #include "Framework/Application/SlateApplication.h"
@@ -51,12 +54,13 @@ FReply FVariantManagerVariantNode::OnDoubleClick(const FGeometry& InMyGeometry,
 
 TSharedRef<SWidget> FVariantManagerVariantNode::GetCustomOutlinerContent(TSharedPtr<SVariantManagerTableRow> InTableRow)
 {
-	FSlateFontInfo NodeFont = FAppStyle::GetFontStyle("Sequencer.AnimationOutliner.RegularFont");
+	const bool bUseNewAssetPicker = UE::VariantManager::Get_CVar_VariantManager_AssetPicker().GetValueOnGameThread();
+
+	FSlateFontInfo NodeFont = FAppStyle::GetFontStyle(bUseNewAssetPicker ? "NormalFont" : "Sequencer.AnimationOutliner.RegularFont");
 
 	EditableLabel = SNew(SInlineEditableTextBlock)
 		.IsReadOnly(this, &FVariantManagerDisplayNode::IsReadOnly)
 		.Font(NodeFont)
-		.ColorAndOpacity(this, &FVariantManagerDisplayNode::GetDisplayNameColor)
 		.OnTextCommitted(this, &FVariantManagerDisplayNode::HandleNodeLabelTextChanged)
 		.Text(this, &FVariantManagerDisplayNode::GetDisplayName)
 		.ToolTipText(this, &FVariantManagerDisplayNode::GetDisplayNameToolTipText)
@@ -92,11 +96,19 @@ TSharedRef<SWidget> FVariantManagerVariantNode::GetCustomOutlinerContent(TShared
 			]
 
 			+ SHorizontalBox::Slot()
-			.Padding(FMargin(4.f, 0.f, 4.f, 0.f))
 			.VAlign(VAlign_Center)
 			.AutoWidth()
 			[
-				GetThumbnailWidget()
+				SNew(SAssetDropTarget)
+				.OnAreAssetsAcceptableForDrop(this, &FVariantManagerVariantNode::AreAssetsAcceptableForDrop)
+				.OnAssetsDropped(this, &FVariantManagerVariantNode::HandleAssetsDropped)
+				[
+					SNew(SBox)
+					.Padding(FMargin(4.f, 0.f, 4.f, 0.f))
+					[
+						GetThumbnailWidget()
+					]
+				]
 			]
 
 			+ SHorizontalBox::Slot()
@@ -471,6 +483,7 @@ void FVariantManagerVariantNode::BuildContextMenu(FMenuBuilder& MenuBuilder)
 	MenuBuilder.EndSection();
 	MenuBuilder.BeginSection(TEXT("Thumbnail"), LOCTEXT("ThumbnailSectionText", "Thumbnail"));
 	MenuBuilder.AddMenuEntry(FVariantManagerEditorCommands::Get().CreateThumbnailCommand);
+	MenuBuilder.AddMenuEntry(FVariantManagerEditorCommands::Get().SelectThumbnailAssetCommand);
 	MenuBuilder.AddMenuEntry(FVariantManagerEditorCommands::Get().LoadThumbnailCommand);
 	MenuBuilder.AddMenuEntry(FVariantManagerEditorCommands::Get().ClearThumbnailCommand);
 	MenuBuilder.EndSection();
@@ -499,6 +512,10 @@ TSharedRef<SWidget> FVariantManagerVariantNode::GetThumbnailWidget()
 			ThumbnailWidget = SNew(SImage).Image(ImageBrush.Get());
 		}
 	}
+	else
+	{
+		ThumbnailWidget = SNew(SImage).Image(FVariantManagerStyle::Get().GetBrush("VariantManager.Variant"));
+	}
 
 	// Use default thumbnail graphic for Variants as a fallback
 	if(!ThumbnailWidget.IsValid())
@@ -515,6 +532,53 @@ TSharedRef<SWidget> FVariantManagerVariantNode::GetThumbnailWidget()
 		];
 }
 
+bool FVariantManagerVariantNode::AreAssetsAcceptableForDrop(TArrayView<FAssetData> AssetDatas) const
+{
+	if (AssetDatas.Num() == 1)
+	{
+		const FAssetData& AssetData = AssetDatas[0];
+		if (UObject* Object = AssetData.GetAsset())
+		{
+			if (Cast<ULevelVariantSets>(Object))
+			{	// We can't take LevelVariantSets as thumbnails
+				return false;
+			}
+
+		    // Try to convert old thumbnails to a new thumbnail
+		    FName ObjectName = *Object->GetFullName();
+		    FThumbnailMap Map;
+		    ThumbnailTools::ConditionallyLoadThumbnailsForObjects({ObjectName}, Map);
+
+		    if (FObjectThumbnail* OldThumbnail = Map.Find(ObjectName))
+		    {
+				return true;
+			}
+		}
+	}
+
+	return false;
+}
+
+void FVariantManagerVariantNode::HandleAssetsDropped(const FDragDropEvent& DragDropEvent, TArrayView<FAssetData> AssetDatas)
+{
+	if (AssetDatas.Num() == 1)
+	{
+		const FAssetData& AssetData = AssetDatas[0];
+		if (UObject* Asset = AssetData.GetAsset())
+		{
+			FScopedTransaction Transaction(LOCTEXT("VariantNodeSetThumbnailTransaction", "Set thumbnail for variant"));
+			Variant.Modify();
+			Variant.SetThumbnailFromObject(Asset);
+			if (TSharedPtr<FVariantManager> VariantManager = GetVariantManager().Pin())
+			{
+				if (TSharedPtr<SVariantManager> VariantManagerWidget = VariantManager->GetVariantManagerWidget())
+				{
+					VariantManagerWidget->RefreshVariantTree();
+				}
+			}
+		}
+	}
+}
 ECheckBoxState FVariantManagerVariantNode::IsRadioButtonChecked() const
 {
 	return Variant.IsActive() ? ECheckBoxState::Checked : ECheckBoxState::Unchecked;
diff --git a/Engine/Plugins/Enterprise/VariantManager/Source/VariantManager/Private/DisplayNodes/VariantManagerVariantNode.h b/Engine/Plugins/Enterprise/VariantManager/Source/VariantManager/Private/DisplayNodes/VariantManagerVariantNode.h
--- a/Engine/Plugins/Enterprise/VariantManager/Source/VariantManager/Private/DisplayNodes/VariantManagerVariantNode.h
+++ b/Engine/Plugins/Enterprise/VariantManager/Source/VariantManager/Private/DisplayNodes/VariantManagerVariantNode.h
@@ -40,6 +40,10 @@ public:
 	UVariant& GetVariant() const;
 	TSharedRef<SWidget> GetThumbnailWidget();
 
+private:
+	bool AreAssetsAcceptableForDrop(TArrayView<FAssetData> AssetDatas) const;
+	void HandleAssetsDropped(const FDragDropEvent& DragDropEvent, TArrayView<FAssetData> AssetDatas);
+
 private:
 
 	ECheckBoxState IsRadioButtonChecked() const;
diff --git a/Engine/Plugins/Enterprise/VariantManager/Source/VariantManager/Private/DisplayNodes/VariantManagerVariantSetNode.cpp b/Engine/Plugins/Enterprise/VariantManager/Source/VariantManager/Private/DisplayNodes/VariantManagerVariantSetNode.cpp
--- a/Engine/Plugins/Enterprise/VariantManager/Source/VariantManager/Private/DisplayNodes/VariantManagerVariantSetNode.cpp
+++ b/Engine/Plugins/Enterprise/VariantManager/Source/VariantManager/Private/DisplayNodes/VariantManagerVariantSetNode.cpp
@@ -14,13 +14,13 @@
 #include "VariantManagerDragDropOp.h"
 #include "VariantManagerEditorCommands.h"
 #include "VariantManagerNodeTree.h"
-#include "VariantManagerNodeTree.h"
 #include "VariantManagerSelection.h"
 #include "VariantManagerStyle.h"
 #include "VariantObjectBinding.h"
 #include "VariantSet.h"
 
 #include "AssetThumbnail.h"
+#include "SAssetDropTarget.h"
 #include "Brushes/SlateImageBrush.h"
 #include "DragAndDrop/ActorDragDropGraphEdOp.h"
 #include "Engine/Texture2D.h"
@@ -63,7 +63,6 @@ TSharedRef<SWidget> FVariantManagerVariantSetNode::GetCustomOutlinerContent(TSha
 	EditableLabel = SNew(SInlineEditableTextBlock)
 		.IsReadOnly(this, &FVariantManagerDisplayNode::IsReadOnly)
 		.Font(NodeFont)
-		.ColorAndOpacity(this, &FVariantManagerDisplayNode::GetDisplayNameColor)
 		.OnTextCommitted(this, &FVariantManagerDisplayNode::HandleNodeLabelTextChanged)
 		.Text(this, &FVariantManagerDisplayNode::GetDisplayName)
 		.ToolTipText(this, &FVariantManagerDisplayNode::GetDisplayNameToolTipText)
@@ -105,11 +104,19 @@ TSharedRef<SWidget> FVariantManagerVariantSetNode::GetCustomOutlinerContent(TSha
 				]
 
 				+ SHorizontalBox::Slot()
-				.Padding(FMargin(4.f, 0.f, 4.f, 0.f))
 				.VAlign(VAlign_Center)
 				.AutoWidth()
 				[
-					GetThumbnailWidget()
+					SNew(SAssetDropTarget)
+					.OnAreAssetsAcceptableForDrop(this, &FVariantManagerVariantSetNode::AreAssetsAcceptableForDrop)
+					.OnAssetsDropped(this, &FVariantManagerVariantSetNode::HandleAssetsDropped)
+					[
+						SNew(SBox)
+						.Padding(FMargin(4.f, 0.f, 4.f, 0.f))
+						[
+							GetThumbnailWidget()
+						]
+					]
 				]
 
 				+ SHorizontalBox::Slot()
@@ -140,7 +147,10 @@ TSharedRef<SWidget> FVariantManagerVariantSetNode::GetCustomOutlinerContent(TSha
 								GetDisplayName()));
 
 							VariantManagerPtr->CreateVariant(&GetVariantSet());
-							VariantManagerPtr->GetVariantManagerWidget()->RefreshVariantTree();
+							if (TSharedPtr<SVariantManager> VariantManagerWidget = VariantManagerPtr->GetVariantManagerWidget())
+							{
+								VariantManagerWidget->RefreshVariantTree();
+							}
 						}
 
 						return FReply::Handled();
@@ -176,6 +186,10 @@ TSharedRef<SWidget> FVariantManagerVariantSetNode::GetThumbnailWidget()
 			ThumbnailWidget = SNew(SImage).Image(ImageBrush.Get());
 		}
 	}
+	else
+	{
+		ThumbnailWidget = SNew(SImage).Image(FVariantManagerStyle::Get().GetBrush("VariantManager.VariantSet"));
+	}
 
 	// Use default thumbnail graphic for Variants as a fallback
 	if (!ThumbnailWidget.IsValid())
@@ -542,7 +556,10 @@ void FVariantManagerVariantSetNode::Drop(const FDragDropEvent& DragDropEvent, EI
 
 		// Add bindings just once, or else it will spawn multiple popups and trigger multiple refreshes
 		VarMan->CreateObjectBindingsAndCaptures(Actors, Vars);
-		VarMan->GetVariantManagerWidget()->RefreshActorList();
+		if (TSharedPtr<SVariantManager> VariantManagerWidget = VarMan->GetVariantManagerWidget())
+		{ 
+			VariantManagerWidget->RefreshActorList();
+		}
 	}
 	else if (VarManDragDrop.IsValid())
 	{
@@ -626,7 +643,10 @@ void FVariantManagerVariantSetNode::Drop(const FDragDropEvent& DragDropEvent, EI
 				SelectedPaths.Add(Binding->GetPathName());
 			}
 
-			VarMan->GetVariantManagerWidget()->RefreshActorList();
+			if (TSharedPtr<SVariantManager> VariantManagerWidget = VarMan->GetVariantManagerWidget())
+			{
+				VariantManagerWidget->RefreshActorList();
+			}
 		}
 		else if (NumDraggedVariants > 0)
 		{
@@ -652,7 +672,10 @@ void FVariantManagerVariantSetNode::Drop(const FDragDropEvent& DragDropEvent, EI
 				SelectedPaths.Add(Variant->GetPathName());
 			}
 
-			VarMan->GetVariantManagerWidget()->RefreshVariantTree();
+			if (TSharedPtr<SVariantManager> VariantManagerWidget = VarMan->GetVariantManagerWidget())
+			{
+				VariantManagerWidget->RefreshVariantTree();
+			}
 		}
 		else if (NumDraggedVariantSets > 0)
 		{
@@ -690,7 +713,10 @@ void FVariantManagerVariantSetNode::Drop(const FDragDropEvent& DragDropEvent, EI
 				}
 			}
 
-			VarMan->GetVariantManagerWidget()->RefreshVariantTree();
+			if (TSharedPtr<SVariantManager> VariantManagerWidget = VarMan->GetVariantManagerWidget())
+			{
+				VariantManagerWidget->RefreshVariantTree();
+			}
 		}
 	}
 }
@@ -714,13 +740,17 @@ void FVariantManagerVariantSetNode::BuildContextMenu(FMenuBuilder& MenuBuilder)
 					GetDisplayName()));
 
 				VariantManagerPtr->CreateVariant(&GetVariantSet());
-				VariantManagerPtr->GetVariantManagerWidget()->RefreshVariantTree();
+				if (TSharedPtr<SVariantManager> VariantManagerWidget = VariantManagerPtr->GetVariantManagerWidget())
+				{
+					VariantManagerWidget->RefreshVariantTree();
+				}
 			}
 		})));
 	MenuBuilder.EndSection();
 
 	MenuBuilder.BeginSection(TEXT("Thumbnail"), LOCTEXT("ThumbnailSectionText", "Thumbnail"));
 	MenuBuilder.AddMenuEntry(FVariantManagerEditorCommands::Get().CreateThumbnailCommand);
+	MenuBuilder.AddMenuEntry(FVariantManagerEditorCommands::Get().SelectThumbnailAssetCommand);
 	MenuBuilder.AddMenuEntry(FVariantManagerEditorCommands::Get().LoadThumbnailCommand);
 	MenuBuilder.AddMenuEntry(FVariantManagerEditorCommands::Get().ClearThumbnailCommand);
 	MenuBuilder.EndSection();
@@ -731,6 +761,54 @@ UVariantSet& FVariantManagerVariantSetNode::GetVariantSet() const
 	return VariantSet;
 }
 
+bool FVariantManagerVariantSetNode::AreAssetsAcceptableForDrop(TArrayView<FAssetData> AssetDatas) const
+{
+	if (AssetDatas.Num() == 1)
+	{
+		const FAssetData& AssetData = AssetDatas[0];
+		if (UObject* Object = AssetData.GetAsset())
+		{
+			if (Cast<ULevelVariantSets>(Object))
+			{	// We can't take LevelVariantSets as thumbnails
+				return false;
+			}
+
+		    // Try to convert old thumbnails to a new thumbnail
+		    FName ObjectName = *Object->GetFullName();
+		    FThumbnailMap Map;
+		    ThumbnailTools::ConditionallyLoadThumbnailsForObjects({ObjectName}, Map);
+
+		    if (FObjectThumbnail* OldThumbnail = Map.Find(ObjectName))
+		    {
+				return true;
+			}
+		}
+	}
+
+	return false;
+}
+
+void FVariantManagerVariantSetNode::HandleAssetsDropped(const FDragDropEvent& DragDropEvent, TArrayView<FAssetData> AssetDatas)
+{
+	if (AssetDatas.Num() == 1)
+	{
+		const FAssetData& AssetData = AssetDatas[0];
+		if (UObject* Asset = AssetData.GetAsset())
+		{
+			FScopedTransaction Transaction(LOCTEXT("VariantSetNodeSetThumbnailTransaction", "Set thumbnail for variant set"));
+			VariantSet.Modify();
+			VariantSet.SetThumbnailFromObject(Asset);
+			if (TSharedPtr<FVariantManager> VariantManager = GetVariantManager().Pin())
+			{
+				if (TSharedPtr<SVariantManager> VariantManagerWidget = VariantManager->GetVariantManagerWidget())
+				{
+					VariantManagerWidget->RefreshVariantTree();
+				}
+			}
+		}
+	}
+}
+
 const FSlateBrush* FVariantManagerVariantSetNode::GetNodeBorderImage() const
 {
 	return IsExpanded() ? ExpandedBackgroundBrush : CollapsedBackgroundBrush;
diff --git a/Engine/Plugins/Enterprise/VariantManager/Source/VariantManager/Private/DisplayNodes/VariantManagerVariantSetNode.h b/Engine/Plugins/Enterprise/VariantManager/Source/VariantManager/Private/DisplayNodes/VariantManagerVariantSetNode.h
--- a/Engine/Plugins/Enterprise/VariantManager/Source/VariantManager/Private/DisplayNodes/VariantManagerVariantSetNode.h
+++ b/Engine/Plugins/Enterprise/VariantManager/Source/VariantManager/Private/DisplayNodes/VariantManagerVariantSetNode.h
@@ -39,6 +39,10 @@ public:
 	/** Gets the folder data for this display node. */
 	UVariantSet& GetVariantSet() const;
 
+private:
+	bool AreAssetsAcceptableForDrop(TArrayView<FAssetData> AssetDatas) const;
+	void HandleAssetsDropped(const FDragDropEvent& DragDropEvent, TArrayView<FAssetData> AssetDatas);
+
 private:
 
 	/** The movie scene folder data which this node represents. */
@@ -53,5 +57,5 @@ private:
 	TSharedPtr<FAssetThumbnailPool> ThumbnailPool;
 
 	// We need to keep the brush alive while the thumbnail is rendered
-	TSharedPtr<FSlateImageBrush> ImageBrush;
+	TSharedPtr<FSlateBrush> ImageBrush;
 };
diff --git a/Engine/Plugins/Enterprise/VariantManager/Source/VariantManager/Private/SDependencyRow.cpp b/Engine/Plugins/Enterprise/VariantManager/Source/VariantManager/Private/SDependencyRow.cpp
--- a/Engine/Plugins/Enterprise/VariantManager/Source/VariantManager/Private/SDependencyRow.cpp
+++ b/Engine/Plugins/Enterprise/VariantManager/Source/VariantManager/Private/SDependencyRow.cpp
@@ -38,7 +38,9 @@ void SDependencyRow::Construct(const FArguments& InArgs, const TSharedRef<STable
 	RebuildVariantSetOptions();
 	RebuildVariantOptions();
 
-	SMultiColumnTableRow<FVariantDependencyModelPtr>::Construct(FSuperRowType::FArguments(), InOwnerTableView);
+	SMultiColumnTableRow<FVariantDependencyModelPtr>::Construct(FSuperRowType::FArguments()
+		.Style(&FAppStyle::Get().GetWidgetStyle<FTableRowStyle>("SceneOutliner.TableViewRow")),
+		InOwnerTableView);
 }
 
 TSharedRef<SWidget> SDependencyRow::GenerateWidgetForColumn(const FName& InColumnName)
diff --git a/Engine/Plugins/Enterprise/VariantManager/Source/VariantManager/Private/SVariantManager.cpp b/Engine/Plugins/Enterprise/VariantManager/Source/VariantManager/Private/SVariantManager.cpp
--- a/Engine/Plugins/Enterprise/VariantManager/Source/VariantManager/Private/SVariantManager.cpp
+++ b/Engine/Plugins/Enterprise/VariantManager/Source/VariantManager/Private/SVariantManager.cpp
@@ -2,6 +2,7 @@
 
 #include "SVariantManager.h"
 
+#include "ConsoleVariables.h"
 #include "Components/StaticMeshComponent.h"
 #include "DisplayNodes/VariantManagerActorNode.h"
 #include "DisplayNodes/VariantManagerColorPropertyNode.h"
@@ -44,12 +45,20 @@
 #include "Misc/TransactionObjectEvent.h"
 #include "SceneOutlinerModule.h"
 #include "ActorTreeItem.h"
+#include "AssetViewUtils.h"
+#include "ContentBrowserModule.h"
+#include "IContentBrowserSingleton.h"
+#include "PropertyCustomizationHelpers.h"
+#include "SAssetDropTarget.h"
 #include "ScopedTransaction.h"
 #include "SPositiveActionButton.h"
+#include "Engine/Texture2D.h"
+#include "Settings/EditorStyleSettings.h"
 #include "Widgets/Input/SButton.h"
 #include "Widgets/Input/SSearchBox.h"
 #include "Widgets/Layout/SScrollBorder.h"
 #include "Widgets/Layout/SSeparator.h"
+#include "Widgets/Layout/SSpacer.h"
 #include "Widgets/Notifications/SNotificationList.h"
 #include "Widgets/Text/SRichTextBlock.h"
 
@@ -63,7 +72,7 @@ FSplitterValues::FSplitterValues(FString& InSerialized)
 	TArray<FString> SplitString;
 	InSerialized.ParseIntoArray(SplitString, TEXT(";"));
 
-	if (SplitString.Num() != 5)
+	if (SplitString.Num() != 3)
 	{
 		return;
 	}
@@ -84,23 +93,169 @@ TSharedRef<SWidget> SVariantManager::MakeAddButton()
 {
 	return SNew(SPositiveActionButton)
 		.Icon(FAppStyle::Get().GetBrush("Icons.Plus"))
-		.Text(LOCTEXT("VariantSetPlusText", "VARIANT SET"))
+		.Text(LOCTEXT("VariantSetPlusText", "Variant Set"))
 		.OnClicked(this, &SVariantManager::OnAddVariantSetClicked);
 }
 
+TSharedRef<SWidget> SVariantManager::GetMenuContextForLevelVariantSets()
+{
+	FMenuBuilder Builder(true, MakeShared<FUICommandList>());
+
+	FAssetData AssetData;
+	if (TSharedPtr<FVariantManager> VariantManager = VariantManagerPtr.Pin())
+	{
+		AssetData = FAssetData(VariantManager->GetCurrentLevelVariantSets());	
+	}
+
+	Builder.AddWidget(PropertyCustomizationHelpers::MakeAssetPickerWithMenu(
+		AssetData,
+		false,
+		{ULevelVariantSets::StaticClass()},
+		{},
+		FOnShouldFilterAsset(),
+		FOnAssetSelected::CreateLambda([this](const FAssetData& AssetData)
+			{
+				if (AssetData.IsValid())
+				{
+					AssetViewUtils::OpenEditorForAsset(AssetData.GetAsset());
+				}
+			}),
+		FSimpleDelegate()),
+		FText::GetEmpty());
+
+	return Builder.MakeWidget();
+}
+
+TSharedRef<SWidget> SVariantManager::MakeLevelVariantSetsComboButton()
+{
+	return SNew(SHorizontalBox)
+		+SHorizontalBox::Slot()
+		.AutoWidth()
+		.VAlign(VAlign_Center)
+		.Padding(2.0f, 2.0f, 8.0f, 2.0f)
+		[
+			SNew(STextBlock)
+			.Text(NSLOCTEXT("SVariantManager", "Asset", "Asset:"))
+		]
+
+		+ SHorizontalBox::Slot()
+		.FillWidth(1.0f)
+		[
+			SNew(SAssetDropTarget)
+			.OnAreAssetsAcceptableForDropWithReason_Static(&SVariantManager::OnAreAssetsAcceptableForVariantManager)
+			.OnAssetsDropped(this, &SVariantManager::OnAssetDroppedForVariantManager)
+			[
+				SNew(SComboButton)
+				.OnGetMenuContent(this, &SVariantManager::GetMenuContextForLevelVariantSets)
+				.ButtonContent()
+				[
+					SNew(SBorder)
+					.BorderBackgroundColor(FSlateColor(FColor::Black))
+					[
+						SNew(SHorizontalBox)
+
+						+ SHorizontalBox::Slot()
+						.AutoWidth()
+						.Padding(2.0f)
+						[
+							SNew(SImage)
+							.Image(FVariantManagerStyle::Get().GetBrush("VariantManager.Icon.Small"))
+							.ColorAndOpacity(FSlateColor::UseForeground())
+							.DesiredSizeOverride(FVector2D(16, 16))
+						]
+
+						+SHorizontalBox::Slot()
+						.FillWidth(1)
+						.VAlign(VAlign_Center)
+						[
+							SNew(STextBlock)
+							.Font(FAppStyle::GetFontStyle(TEXT("PropertyWindow.NormalFont")))
+							.Text_Lambda([this]() -> FText
+							{
+								TSharedPtr<FVariantManager> VarMan = VariantManagerPtr.Pin();
+								if (VarMan.IsValid() && VarMan->GetCurrentLevelVariantSets())
+								{
+									return FText::FromName(VarMan->GetCurrentLevelVariantSets()->GetFName());
+								}
+
+								return NSLOCTEXT("SVariantManager", "NoAssetSelected", "No asset selected");
+							})
+			                .OverflowPolicy(GetDefault<UEditorStyleSettings>()->bEnableMiddleEllipsis ? TOptional(ETextOverflowPolicy::MiddleEllipsis) : NullOpt)
+						]
+					]
+				]
+			]
+		]
+
+		+ SHorizontalBox::Slot()
+		.AutoWidth()
+		.Padding(2.0f, 0.0f)
+		[
+			SNew(SBox)
+			.HAlign(HAlign_Center)
+			.VAlign(VAlign_Center)
+			.WidthOverride(22.0f)
+			.HeightOverride(22.0f)
+			.ToolTipText(LOCTEXT( "UseButtonToolTipText", "Use Selected Asset from Content Browser"))
+			[
+				SNew(SButton)
+				.ButtonStyle( FAppStyle::Get(), "SimpleButton" )
+				.OnClicked_Static(&SVariantManager::OnUseSelectedAssetFromContentBrowserForCurrentLevelVariantSet)
+				.ContentPadding(0.0f)
+				.IsFocusable(true)
+				[ 
+					SNew( SImage )
+					.Image(FAppStyle::GetBrush("Icons.Use"))
+					.ColorAndOpacity( FSlateColor::UseForeground() )
+				]
+			]
+		]
+		
+		+ SHorizontalBox::Slot()
+		.AutoWidth()
+		.Padding(2.0f, 0.0f)
+		[
+			SNew(SBox)
+			.HAlign(HAlign_Center)
+			.VAlign(VAlign_Center)
+			.WidthOverride(22.0f)
+			.HeightOverride(22.0f)
+			.ToolTipText(LOCTEXT( "BrowseButtonToolTipText", "Browse to Asset in Content Browser"))
+			[
+				SNew(SButton)
+				.ButtonStyle( FAppStyle::Get(), "SimpleButton" )
+				.OnClicked( this, &SVariantManager::OnBrowseToCurrentLevelVariantSet)
+				.ContentPadding(0.0f)
+				.IsFocusable(true)
+				[ 
+					SNew( SImage )
+					.Image(FAppStyle::Get().GetBrush("Icons.BrowseContent"))
+					.ColorAndOpacity( FSlateColor::UseForeground() )
+				]
+			]
+		];
+}
+
 TSharedRef<ITableRow> SVariantManager::MakeCapturedPropertyRow(TSharedPtr<FVariantManagerPropertyNode> Item, const TSharedRef<STableViewBase>& OwnerTable)
 {
 	return SNew(SVariantManagerTableRow, OwnerTable, StaticCastSharedPtr<FVariantManagerDisplayNode>(Item).ToSharedRef());
 }
 
-TSharedPtr<SWidget> SVariantManager::OnPropertyListContextMenuOpening()
+TSharedRef<ITableRow> SVariantManager::MakeCapturedPropertyRowWithColumns(
+	TSharedPtr<FVariantManagerPropertyNode> Item,
+	const TSharedRef<STableViewBase>& OwnerTable)
+{
+	return SNew(SVariantManagerMultiColumnPropertyTableRow, OwnerTable, Item.ToSharedRef());
+}
+
+TSharedPtr<SWidget> SVariantManager::OnPropertyListContextMenuOpening() const
 {
 	TArray<TSharedPtr<FVariantManagerPropertyNode>> SelectedNodes;
 	CapturedPropertyListView->GetSelectedItems(SelectedNodes);
 
 	if (SelectedNodes.Num() > 0)
 	{
-		const bool bShouldCloseWindowAfterMenuSelection = true;
+		constexpr bool bShouldCloseWindowAfterMenuSelection = true;
 		FMenuBuilder MenuBuilder(bShouldCloseWindowAfterMenuSelection, GetPropertyListCommandBindings());
 
 		SelectedNodes[0]->BuildContextMenu(MenuBuilder);
@@ -228,6 +383,17 @@ void SVariantManager::Construct(const FArguments& InArgs, TSharedRef<FVariantMan
 
 	const float BorderThickness = FVariantManagerStyle::Get().GetFloat("VariantManager.Spacings.BorderThickness");;
 
+	const bool bUseNewAssetPicker = UE::VariantManager::Get_CVar_VariantManager_AssetPicker().GetValueOnGameThread();
+	TSharedRef<SWidget> VariantSetsComboBox = bUseNewAssetPicker ? MakeLevelVariantSetsComboButton() : SNullWidget::NullWidget;
+	TSharedRef<SWidget> OldSearchBox = bUseNewAssetPicker ? SNullWidget::NullWidget : SNew(SSearchBox)
+		.HintText(LOCTEXT("VariantManagerFilterText", "Search Variant"))
+		.OnTextChanged(this, &SVariantManager::OnOutlinerSearchChanged)
+		.MinDesiredWidth(200);
+	TSharedRef<SWidget> NewSearchBox = !bUseNewAssetPicker ? SNullWidget::NullWidget : SNew(SSearchBox)
+		.HintText(LOCTEXT("VariantManagerFilterText", "Search Variant"))
+		.OnTextChanged(this, &SVariantManager::OnOutlinerSearchChanged)
+		.MinDesiredWidth(200);
+
 	ChildSlot
 	[
 		SNew(SVerticalBox)
@@ -249,20 +415,25 @@ void SVariantManager::Construct(const FArguments& InArgs, TSharedRef<FVariantMan
 				.Padding(FMargin(0.f, 0.f, VM_COMMON_PADDING+2.0f, 1.f))
 				.AutoWidth()
 				[
-					MakeAddButton()
+					VariantSetsComboBox
 				]
 
+				+SHorizontalBox::Slot()
+				.VAlign(VAlign_Center)
+				.Padding(FMargin(0.f, 0.f, VM_COMMON_PADDING+2.0f, 1.f))
+				.AutoWidth()
+				[
+					MakeAddButton()
+				]
+					
 				+SHorizontalBox::Slot()
 				.FillWidth(1.0f)
 				.HAlign(HAlign_Left)
 				.VAlign(VAlign_Center)
 				[
-					SNew(SSearchBox)
-					.HintText(LOCTEXT("VariantManagerFilterText", "Search Variant"))
-					.OnTextChanged(this, &SVariantManager::OnOutlinerSearchChanged)
-					.MinDesiredWidth(200)
+					OldSearchBox
 				]
-
+				
 				+SHorizontalBox::Slot()
 				.VAlign(VAlign_Center)
 				.HAlign(HAlign_Right)
@@ -340,11 +511,26 @@ void SVariantManager::Construct(const FArguments& InArgs, TSharedRef<FVariantMan
 				+SSplitter::Slot()
 				.Value(SplitterValues.VariantColumn)
 				[
-					SNew(SScrollBorder, NodeTreeView.ToSharedRef())
+					SNew(SVerticalBox)
+
+					+ SVerticalBox::Slot()
+					.AutoHeight()
+					.HAlign(HAlign_Left)
+					.VAlign(VAlign_Center)
+					.Padding(FMargin(10, 7, 8, 8))
 					[
-						SNew(SBox) // Very important to prevent the tree from expanding freely
+						NewSearchBox
+					]
+
+					+ SVerticalBox::Slot()
+					.FillHeight(1.0f)
+					[
+						SNew(SScrollBorder, NodeTreeView.ToSharedRef())
 						[
-							NodeTreeView.ToSharedRef()
+							SNew(SBox) // Very important to prevent the tree from expanding freely
+							[
+								NodeTreeView.ToSharedRef()
+							]
 						]
 					]
 				]
@@ -372,6 +558,8 @@ void SVariantManager::Construct(const FArguments& InArgs, TSharedRef<FVariantMan
 		]
 	];
 
+	SortItems();
+
 	RefreshVariantTree();
 }
 
@@ -490,6 +678,12 @@ void SVariantManager::CreateCommandBindings()
 		FExecuteAction::CreateSP(this, &SVariantManager::CreateThumbnail),
 		FCanExecuteAction::CreateSP(this, &SVariantManager::CanCreateThumbnail)
 	);
+	
+	VariantTreeCommandBindings->MapAction(
+		FVariantManagerEditorCommands::Get().SelectThumbnailAssetCommand,
+		FExecuteAction::CreateSP(this, &SVariantManager::SelectThumbnailAsset),
+		FCanExecuteAction::CreateSP(this, &SVariantManager::CanLoadThumbnail)
+	);
 
 	VariantTreeCommandBindings->MapAction(
 		FVariantManagerEditorCommands::Get().LoadThumbnailCommand,
@@ -1185,6 +1379,50 @@ void SVariantManager::CreateThumbnail()
 	}
 }
 
+void SVariantManager::SelectThumbnailAsset()
+{
+	TSharedPtr<FVariantManager> VariantManager = VariantManagerPtr.Pin();
+	if (!VariantManager.IsValid())
+	{
+		return;
+	}
+
+	TArray<UVariant*> Vars;
+	TArray<UVariantSet*> VarSets;
+	VariantManager->GetSelection().GetSelectedVariantsAndVariantSets(Vars, VarSets);
+	if (Vars.Num() + VarSets.Num() < 1)
+	{
+		return;
+	}
+
+	IContentBrowserSingleton& ContentBrowser = FModuleManager::LoadModuleChecked<FContentBrowserModule>("ContentBrowser").Get();
+	FOpenAssetDialogConfig Config;
+	Config.bAllowMultipleSelection = false;
+	Config.DialogTitleOverride = NSLOCTEXT("SVariantManager", "SelectAnAssetForThumbnail", "Select an Asset for Thumbnail");
+	TArray<FAssetData> AssetDatas = ContentBrowser.CreateModalOpenAssetDialog(Config);
+	if (AssetDatas.Num() == 1)
+	{
+		if (UObject* ObjectThumbnail = AssetDatas[0].GetAsset())
+		{
+			FScopedTransaction Transaction(LOCTEXT("SelectThumbnailAssetTransaction", "Select asset as a thumbnail for a variant or variant set"));
+
+			for (UVariant* Var : Vars)
+			{
+				Var->SetThumbnailFromObject(ObjectThumbnail);
+
+				if (UVariantSet* VarSet = Var->GetParent())
+				{
+					UpdateVariantSetThumbnail(VarSet);
+				}
+			}
+			for (UVariantSet* VarSet : VarSets)
+			{
+				VarSet->SetThumbnailFromObject(ObjectThumbnail);
+			}
+		}
+	}
+}
+
 void SVariantManager::LoadThumbnail()
 {
 	TSharedPtr<FVariantManager> VariantManager = VariantManagerPtr.Pin();
@@ -1287,6 +1525,82 @@ void SVariantManager::UpdateVariantSetThumbnail(UVariantSet* InVariantSet)
 	}
 }
 
+void SVariantManager::OnAssetDroppedForVariantManager(const FDragDropEvent& DragDropEvent, TArrayView<FAssetData> AssetDatas)
+{
+	if (AssetDatas.IsValidIndex(0))
+	{
+		const FAssetData& AssetData = AssetDatas[0];
+		if (AssetData.AssetClassPath != FTopLevelAssetPath(ULevelVariantSets::StaticClass()))
+		{
+			return;
+		}
+		UObject* Asset = AssetData.GetAsset();
+		ULevelVariantSets* LevelVariantSets = Cast<ULevelVariantSets>(Asset);
+		if (!LevelVariantSets)
+		{
+			return;
+		}
+		AssetViewUtils::OpenEditorForAsset(LevelVariantSets);
+	}
+}
+
+bool SVariantManager::OnAreAssetsAcceptableForVariantManager(TArrayView<FAssetData> Assets, FText& OutReason)
+{
+	if (Assets.Num() != 1)
+	{
+		OutReason = LOCTEXT("MultipleAssetsNotAllowed", "Please select only one asset.");
+		return false;
+	}
+
+	FAssetData AssetData = Assets[0];
+
+	if (AssetData.AssetClassPath != FTopLevelAssetPath(ULevelVariantSets::StaticClass()))
+	{
+		OutReason = LOCTEXT("LevelVariantSetsNotAllowed", "Only Level Variant Sets assets can be opened by Variant Manager.");
+		return false;
+	}
+
+	return true;
+}
+
+FReply SVariantManager::OnUseSelectedAssetFromContentBrowserForCurrentLevelVariantSet()
+{
+	IContentBrowserSingleton& Singleton = FModuleManager::LoadModuleChecked<FContentBrowserModule>("ContentBrowser").Get();;
+	TArray<FAssetData> ContentBrowserSelectedAssets;
+	Singleton.GetSelectedAssets(ContentBrowserSelectedAssets);
+
+	if (ContentBrowserSelectedAssets.Num() == 1)
+	{
+		FAssetData SelectedAsset = ContentBrowserSelectedAssets[0];
+
+		if (SelectedAsset.IsInstanceOf(ULevelVariantSets::StaticClass()))
+		{
+			if (ULevelVariantSets* LevelVariantSets = Cast<ULevelVariantSets>(SelectedAsset.GetAsset()))
+			{
+				AssetViewUtils::OpenEditorForAsset(LevelVariantSets);
+			}
+		}
+	}
+
+	return FReply::Handled();
+}
+
+FReply SVariantManager::OnBrowseToCurrentLevelVariantSet() const
+{
+	if (TSharedPtr<FVariantManager> VariantManger = VariantManagerPtr.Pin())
+	{
+		if (ULevelVariantSets* LevelVariantSets = VariantManger->GetCurrentLevelVariantSets())
+		{
+			TArray<FAssetData> AssetDataArray;
+			AssetDataArray.Add(FAssetData(LevelVariantSets));
+			FContentBrowserModule& ContentBrowserModule = FModuleManager::LoadModuleChecked<FContentBrowserModule>("ContentBrowser");
+			ContentBrowserModule.Get().SyncBrowserToAssets(AssetDataArray);
+		}
+	}
+
+	return FReply::Handled();
+}
+
 bool SVariantManager::CanSwitchOnVariant()
 {
 	TSharedPtr<FVariantManager> VariantManager = VariantManagerPtr.Pin();
@@ -3412,8 +3726,9 @@ TSharedRef<ITableRow> SVariantManager::GenerateRightTreeRow( TSharedRef<ERightTr
 
 TSharedRef<SWidget> SVariantManager::GenerateRightTreeHeaderRowContent( ERightTreeRowType RowType, TSharedRef<STableRow<TSharedRef<ERightTreeRowType>>> InTableRow )
 {
-	return
-		SNew(SBox)
+	const bool bUseNewAssetPicker = UE::VariantManager::Get_CVar_VariantManager_AssetPicker().GetValueOnGameThread();
+
+	return SNew(SBox)
 		[
 			SNew(SBorder)
 			.VAlign( VAlign_Center )
@@ -3450,11 +3765,14 @@ TSharedRef<SWidget> SVariantManager::GenerateRightTreeHeaderRowContent( ERightTr
 				.HAlign(HAlign_Left)
 				.FillWidth(1.0f)
 				[
-					SNew(SRichTextBlock)
-					.Text(LOCTEXT("Dependencies", "Dependencies"))
-					.TransformPolicy(ETextTransformPolicy::ToUpper)
-					.DecoratorStyleSet(&FAppStyle::Get())
-					.TextStyle(FAppStyle::Get(), "DetailsView.CategoryTextStyle")
+					bUseNewAssetPicker ?
+						SNew(SRichTextBlock)
+						.Text(LOCTEXT("Dependencies", "Dependencies"))
+					:
+						SNew(SRichTextBlock)
+						.Text(LOCTEXT("Dependencies", "Dependencies"))
+						.DecoratorStyleSet(&FAppStyle::Get())
+						.TextStyle(FAppStyle::Get(), "DetailsView.CategoryTextStyle")
 				]
 
 				+ SHorizontalBox::Slot()
@@ -3540,6 +3858,68 @@ TSharedRef<SWidget> SVariantManager::GenerateRightTreePropertiesRowContent()
 {
 	float BorderThickness = FVariantManagerStyle::Get().GetFloat("VariantManager.Spacings.BorderThickness");
 
+	const bool bShowColumns = UE::VariantManager::Get_CVar_VariantManager_AssetPicker().GetValueOnAnyThread();
+
+	TSharedPtr<SWidget> PropertyListWidget;
+
+	float ActorsSplitterPrecent = 0.5f;
+	float PropertiesSplitterPercent = 0.5f;
+
+	if (bShowColumns)
+	{
+		ActorsSplitterPrecent = 0.2f;
+		PropertiesSplitterPercent = 0.8f;
+
+		PropertyListWidget = SNew(SVerticalBox)
+
+			+ SVerticalBox::Slot()
+			.AutoHeight()
+			[
+				SNew(SSeparator)
+			]
+			+ SVerticalBox::Slot()
+			.FillHeight(1.0f)
+			[
+				SAssignNew(CapturedPropertyListView, SListView<TSharedPtr<FVariantManagerPropertyNode>>)
+					.SelectionMode(ESelectionMode::Single)
+					.ListItemsSource(&DisplayedPropertyNodes)
+					.OnContextMenuOpening(this, &SVariantManager::OnPropertyListContextMenuOpening)
+					.OnGenerateRow_Static(&SVariantManager::MakeCapturedPropertyRowWithColumns)
+					.Visibility(EVisibility::Visible)
+					.HeaderRow(
+						SNew(SHeaderRow)
+
+						+ SHeaderRow::Column(SVariantManagerMultiColumnPropertyTableRow::ColumnPath)
+						.DefaultLabel(LOCTEXT("PathColumnHeader", "Path"))
+						.SortMode(this, &SVariantManager::GetSortMode, SVariantManagerMultiColumnPropertyTableRow::ColumnPath)
+						.SortPriority(this,  &SVariantManager::GetSortPriority, SVariantManagerMultiColumnPropertyTableRow::ColumnPath)
+						.OnSort(this, &SVariantManager::HandleSort)
+						.FillWidth(0.2f)
+
+						+ SHeaderRow::Column(SVariantManagerMultiColumnPropertyTableRow::ColumnProperty)
+						.DefaultLabel(LOCTEXT("PropertyColumnHeader", "Property"))
+						.SortMode(this, &SVariantManager::GetSortMode, SVariantManagerMultiColumnPropertyTableRow::ColumnProperty)
+						.SortPriority(this,  &SVariantManager::GetSortPriority, SVariantManagerMultiColumnPropertyTableRow::ColumnProperty)
+						.OnSort(this, &SVariantManager::HandleSort)
+						.FillWidth(0.3f)
+
+						+ SHeaderRow::Column(SVariantManagerMultiColumnPropertyTableRow::ColumnValue)
+						.DefaultLabel(LOCTEXT("ValueColumnHeader", "Value"))
+						.FillWidth(0.5f)
+					)
+			];
+	}
+	else
+	{
+		PropertyListWidget = SAssignNew(CapturedPropertyListView, SListView<TSharedPtr<FVariantManagerPropertyNode>>)
+			.SelectionMode(ESelectionMode::Single)
+			.ListItemsSource(&DisplayedPropertyNodes)
+			.OnContextMenuOpening(this, &SVariantManager::OnPropertyListContextMenuOpening)
+			.OnGenerateRow_Static(&SVariantManager::MakeCapturedPropertyRow)
+			.Visibility(EVisibility::Visible);
+	}
+
+
 	return SNew( SBox )
 	[
 		// Actor column
@@ -3549,6 +3929,7 @@ TSharedRef<SWidget> SVariantManager::GenerateRightTreePropertiesRowContent()
 		.HitDetectionSplitterHandleSize( BorderThickness )
 
 		+ SSplitter::Slot()
+		.Value(ActorsSplitterPrecent)
 		[
 			SNew(SVerticalBox)
 			+ SVerticalBox::Slot()
@@ -3624,6 +4005,7 @@ TSharedRef<SWidget> SVariantManager::GenerateRightTreePropertiesRowContent()
 
 		// Properties column
 		+ SSplitter::Slot()
+		.Value(PropertiesSplitterPercent)
 		[
 			SNew(SVerticalBox)
 			+ SVerticalBox::Slot()
@@ -3678,12 +4060,7 @@ TSharedRef<SWidget> SVariantManager::GenerateRightTreePropertiesRowContent()
 			.VAlign(VAlign_Fill)
 			.FillHeight(1.0f)
 			[
-				SAssignNew(CapturedPropertyListView, SListView<TSharedPtr<FVariantManagerPropertyNode>>)
-				.SelectionMode(ESelectionMode::Single)
-				.ListItemsSource(&DisplayedPropertyNodes)
-				.OnContextMenuOpening(this, &SVariantManager::OnPropertyListContextMenuOpening)
-				.OnGenerateRow(this, &SVariantManager::MakeCapturedPropertyRow)
-				.Visibility(EVisibility::Visible)
+				PropertyListWidget.ToSharedRef()
 			]
 		]
 	];
@@ -3850,6 +4227,64 @@ void SVariantManager::OnThumbnailChanged(UObject* VariantOrVariantSet)
 	}
 }
 
+EColumnSortMode::Type SVariantManager::GetSortMode(const FName ColumnName) const
+{
+	return ColumnName == SortColumn ? SortMode : EColumnSortMode::None;
+}
+
+EColumnSortPriority::Type SVariantManager::GetSortPriority(const FName ColumnName) const
+{
+	return ColumnName == SortColumn ? SortPriority : EColumnSortPriority::None;
+}
+
+void SVariantManager::HandleSort(
+	const EColumnSortPriority::Type NewSortPriority, 
+	const FName& ColumnId,
+	const EColumnSortMode::Type NewSortMode)
+{
+	if (NewSortPriority == EColumnSortPriority::Primary)
+	{
+		SortColumn = ColumnId;
+		SortMode = NewSortMode;
+		SortPriority = NewSortPriority;
+	}
+
+	SortItems();
+
+	if (CapturedPropertyListView)
+	{
+		CapturedPropertyListView->RequestListRefresh();
+	}
+}
+
+void SVariantManager::SortItems()
+{
+	auto SortFn = [this](const TSharedPtr<FVariantManagerPropertyNode>& A, const TSharedPtr<FVariantManagerPropertyNode>& B)
+		{
+			if (SortColumn == SVariantManagerMultiColumnPropertyTableRow::ColumnPath)
+			{
+				if (A->GetSubPath() != B->GetSubPath())
+				{
+					const bool bRet = A->GetSubPath() < B->GetSubPath();
+
+					return SortMode == EColumnSortMode::Ascending ? bRet : !bRet;
+				}
+			}
+			else if (SortColumn == SVariantManagerMultiColumnPropertyTableRow::ColumnProperty)
+			{
+				if (A->GetPropertyName() != B->GetPropertyName())
+				{
+					const bool bRet = A->GetPropertyName() < B->GetPropertyName();
+
+					return SortMode == EColumnSortMode::Ascending ? bRet : !bRet;
+				}
+			}
+
+			return false;
+		};
+	DisplayedPropertyNodes.Sort(SortFn);
+}
+
 #undef VM_COMMON_PADDING
 #undef VM_COMMON_HEADER_MAX_HEIGHT
 #undef LOCTEXT_NAMESPACE
diff --git a/Engine/Plugins/Enterprise/VariantManager/Source/VariantManager/Private/SVariantManager.h b/Engine/Plugins/Enterprise/VariantManager/Source/VariantManager/Private/SVariantManager.h
--- a/Engine/Plugins/Enterprise/VariantManager/Source/VariantManager/Private/SVariantManager.h
+++ b/Engine/Plugins/Enterprise/VariantManager/Source/VariantManager/Private/SVariantManager.h
@@ -3,8 +3,10 @@
 #pragma once
 
 #include "Misc/NotifyHook.h"
-#include "Widgets/SCompoundWidget.h"
+#include "Widgets/Input/SComboButton.h"
 #include "Widgets/Layout/SSplitter.h"
+#include "Widgets/SCompoundWidget.h"
+#include "Widgets/Views/SHeaderRow.h"
 
 class FUICommandList;
 class SScrollBar;
@@ -20,7 +22,6 @@ class FExtender;
 class FTransactionObjectEvent;
 class FVariantManager;
 class FVariantManagerDisplayNode;
-class FVariantManagerDisplayNode;
 class FVariantManagerPropertyNameNode;
 class FVariantManagerPropertyNode;
 class ITableRow;
@@ -128,6 +129,7 @@ public:
 
 	void SwitchOnSelectedVariant();
 	void CreateThumbnail();
+	void SelectThumbnailAsset();
 	void LoadThumbnail();
 	void ClearThumbnail();
 
@@ -169,14 +171,17 @@ public:
 	void SortDisplayNodes(TArray<TSharedRef<FVariantManagerDisplayNode>>& DisplayNodes);
 
 	TSharedRef<SWidget> MakeAddButton();
+	TSharedRef<SWidget> GetMenuContextForLevelVariantSets();
+	TSharedRef<SWidget> MakeLevelVariantSetsComboButton();
 
 	FVariantManagerPropertiesColumnSizeData& GetPropertiesColumnSizeData()
 	{
 		return PropertiesColumnSizeData;
 	}
 
-	TSharedRef<ITableRow> MakeCapturedPropertyRow(TSharedPtr<FVariantManagerPropertyNode> Item, const TSharedRef<STableViewBase>& OwnerTable);
-	TSharedPtr<SWidget> OnPropertyListContextMenuOpening();
+	static TSharedRef<ITableRow> MakeCapturedPropertyRow(TSharedPtr<FVariantManagerPropertyNode> Item, const TSharedRef<STableViewBase>& OwnerTable);
+	static TSharedRef<ITableRow> MakeCapturedPropertyRowWithColumns(TSharedPtr<FVariantManagerPropertyNode> Item, const TSharedRef<STableViewBase>& OwnerTable);
+	TSharedPtr<SWidget> OnPropertyListContextMenuOpening() const;
 
 	void OnOutlinerNodeSelectionChanged();
 	void OnActorNodeSelectionChanged();
@@ -245,6 +250,17 @@ private:
 	// variant thumbnail
 	void UpdateVariantSetThumbnail(UVariantSet* InVariantSet);
 
+	void OnAssetDroppedForVariantManager(const FDragDropEvent& DragDropEvent, TArrayView<FAssetData> AssetDatas);
+	static bool OnAreAssetsAcceptableForVariantManager(TArrayView<FAssetData> Assets, FText& OutReason);
+
+	static FReply OnUseSelectedAssetFromContentBrowserForCurrentLevelVariantSet();
+	FReply OnBrowseToCurrentLevelVariantSet() const;
+
+	EColumnSortMode::Type GetSortMode(const FName ColumnName) const;
+	EColumnSortPriority::Type GetSortPriority(const FName ColumnName) const;
+	void HandleSort(const EColumnSortPriority::Type SortPriority, const FName& ColumnId, const EColumnSortMode::Type NewSortMode);
+	void SortItems();
+
 private:
 	TWeakPtr<FVariantManager> VariantManagerPtr;
 
@@ -275,8 +291,8 @@ private:
 	bool bAutoCaptureProperties = false;
 
 	FVariantManagerPropertiesColumnSizeData PropertiesColumnSizeData;
-	float PropertiesNameColumnWidth;
-	float PropertiesValueColumnWidth;
+	float PropertiesNameColumnWidth = 0;
+	float PropertiesValueColumnWidth = 0;
 
 	FSplitterValues SplitterValues;
 
@@ -316,4 +332,8 @@ private:
 
 	const FTableColumnHeaderStyle* HeaderStyle = nullptr;
 	const FSplitterStyle* SplitterStyle = nullptr;
+
+	EColumnSortMode::Type SortMode = EColumnSortMode::Ascending;
+	EColumnSortPriority::Type SortPriority = EColumnSortPriority::Primary;
+	FName SortColumn;
 };
diff --git a/Engine/Plugins/Enterprise/VariantManager/Source/VariantManager/Private/SVariantManagerTableRow.cpp b/Engine/Plugins/Enterprise/VariantManager/Source/VariantManager/Private/SVariantManagerTableRow.cpp
--- a/Engine/Plugins/Enterprise/VariantManager/Source/VariantManager/Private/SVariantManagerTableRow.cpp
+++ b/Engine/Plugins/Enterprise/VariantManager/Source/VariantManager/Private/SVariantManagerTableRow.cpp
@@ -35,94 +35,272 @@ void SVariantManagerTableRow::Construct(const FArguments& InArgs, const TSharedR
 
 FReply SVariantManagerTableRow::DragDetected( const FGeometry& InGeometry, const FPointerEvent& InPointerEvent )
 {
-	TSharedPtr<FVariantManagerDisplayNode> PinnedNode = Node.Pin();
-	TSharedPtr<FVariantManager> VarMan = PinnedNode->GetVariantManager().Pin();
+	if (TSharedPtr<FVariantManagerDisplayNode> PinnedNode = Node.Pin())
+	{
+		if (TSharedPtr<FVariantManager> VarMan = PinnedNode->GetVariantManager().Pin())
+		{
+			const FSlateBrush* Icon = PinnedNode->GetIconBrush();
 
-	const FSlateBrush* Icon = PinnedNode->GetIconBrush();
+			FVariantManagerSelection& Selection = VarMan->GetSelection();
 
-	if (PinnedNode.IsValid() && VarMan.IsValid())
-	{
-		FVariantManagerSelection& Selection = VarMan->GetSelection();
+			TArray<FDisplayNodeRef> DraggableNodes;
+			FText DefaultHoverText;
 
-		TArray<FDisplayNodeRef> DraggableNodes;
-		FText DefaultHoverText;
+			// We'll drag a group of nodes based on what type we are (e.g. if we're variant or variant
+			// set, drag all of those)
+			switch (PinnedNode->GetType())
+			{
+			case EVariantManagerNodeType::Actor:
+			{
+				for (const FDisplayNodeRef SelectedNode : Selection.GetSelectedActorNodes())
+				{
+					if (SelectedNode->CanDrag())
+					{
+						DraggableNodes.Add(SelectedNode);
+					}
+				}
 
-		// We'll drag a group of nodes based on what type we are (e.g. if we're variant or variant
-		// set, drag all of those)
-		switch (PinnedNode->GetType())
-		{
-		case EVariantManagerNodeType::Actor:
-		{
-			for (const FDisplayNodeRef SelectedNode : Selection.GetSelectedActorNodes())
+				int32 NumNodes = DraggableNodes.Num();
+				if (NumNodes > 1)
+				{
+					Icon = FSlateIconFinder::FindIconForClass(AActor::StaticClass()).GetOptionalIcon();
+				}
+
+				DefaultHoverText = FText::Format(NSLOCTEXT("VariantManagerTableRow", "DragActorNode", "{0} actor {0}|plural(one=binding,other=bindings)" ), NumNodes);
+				break;
+			}
+			case EVariantManagerNodeType::Function:
+			case EVariantManagerNodeType::Property:
+			{
+				if (PinnedNode.IsValid() && PinnedNode->CanDrag())
+				{
+					// We only support dragging one node property/function nodes for now
+					ensure(DraggableNodes.Num() == 0);
+
+					DraggableNodes.Add(PinnedNode.ToSharedRef());
+					DefaultHoverText = FText::Format(NSLOCTEXT("VariantManagerTableRow", "DragPropertyNode", "{0}"), PinnedNode->GetDisplayName());
+				}
+				break;
+			}
+			case EVariantManagerNodeType::Variant:  // Intended fallthrough
+			case EVariantManagerNodeType::VariantSet:
 			{
-				if (SelectedNode->CanDrag())
+				for (const FDisplayNodeRef& SelectedNode : Selection.GetSelectedOutlinerNodes())
 				{
-					DraggableNodes.Add(SelectedNode);
+					if (SelectedNode->CanDrag())
+					{
+						DraggableNodes.Add(SelectedNode);
+					}
 				}
+
+				int32 NumNodes = DraggableNodes.Num();
+
+				DefaultHoverText = FText::Format(NSLOCTEXT("VariantManagerTableRow", "DragVariants", "{0} {0}|plural(one=variant,other=variants) and/or variant {0}|plural(one=set,other=sets)" ),
+					NumNodes);
+
+				break;
+			}
+			default:
+				break;
+			}
+
+			if (DraggableNodes.Num() == 0)
+			{
+				return FReply::Unhandled();
 			}
 
-			int32 NumNodes = DraggableNodes.Num();
-			if (NumNodes > 1)
+			if (TSharedPtr<SVariantManager> VariantManagerWidget = VarMan->GetVariantManagerWidget())
 			{
-				Icon = FSlateIconFinder::FindIconForClass(AActor::StaticClass()).GetOptionalIcon();
+				VariantManagerWidget->SortDisplayNodes(DraggableNodes);
 			}
 
-			DefaultHoverText = FText::Format(NSLOCTEXT("VariantManagerTableRow", "DragActorNode", "{0} actor {0}|plural(one=binding,other=bindings)" ), NumNodes);
-			break;
+			TSharedRef<FVariantManagerDragDropOp> DragDropOp = FVariantManagerDragDropOp::New(DraggableNodes);
+			DragDropOp->SetToolTip(DefaultHoverText, Icon);
+			DragDropOp->SetupDefaults();
+
+			return FReply::Handled().BeginDragDrop(DragDropOp);
 		}
-		case EVariantManagerNodeType::Function:
-		case EVariantManagerNodeType::Property:
+	}
+
+	return FReply::Unhandled();
+}
+
+void SVariantManagerTableRow::DragLeave(const FDragDropEvent& DragDropEvent)
+{
+	TSharedPtr<FDecoratedDragDropOp> DecoratedDragDropOp = DragDropEvent.GetOperationAs<FDecoratedDragDropOp>();
+	if (DecoratedDragDropOp.IsValid())
+	{
+		DecoratedDragDropOp->ResetToDefaultToolTip();
+	}
+}
+
+TOptional<EItemDropZone> SVariantManagerTableRow::CanAcceptDrop( const FDragDropEvent& DragDropEvent, EItemDropZone InItemDropZone, FDisplayNodeRef DisplayNode )
+{
+	TSharedPtr<FVariantManagerDisplayNode> PinnedNode = Node.Pin();
+	if (PinnedNode.IsValid())
+	{
+		return PinnedNode->CanDrop(DragDropEvent, InItemDropZone);
+	}
+
+	return TOptional<EItemDropZone>();
+}
+
+FReply SVariantManagerTableRow::AcceptDrop( const FDragDropEvent& DragDropEvent, EItemDropZone InItemDropZone, FDisplayNodeRef DisplayNode )
+{
+	TSharedPtr<FVariantManagerDisplayNode> PinnedNode = Node.Pin();
+	if (PinnedNode.IsValid())
+	{
+		PinnedNode->Drop(DragDropEvent, InItemDropZone);
+		return FReply::Handled();
+	}
+
+	return FReply::Unhandled();
+}
+
+FReply SVariantManagerTableRow::OnMouseButtonDoubleClick(const FGeometry& InMyGeometry, const FPointerEvent& InMouseEvent)
+{
+	STableRow::OnMouseButtonDoubleClick(InMyGeometry, InMouseEvent);
+
+	if (TSharedPtr<FVariantManagerDisplayNode> PinnedNode = Node.Pin())
+	{
+		return PinnedNode->OnDoubleClick(InMyGeometry, InMouseEvent);
+	}
+
+	return FReply::Unhandled();
+}
+
+// Small hack to bypass CanDrop calls to spacer nodes, letting the underlying tree handle the events instead
+FReply SVariantManagerTableRow::OnDragOver(const FGeometry& MyGeometry, const FDragDropEvent& DragDropEvent)
+{
+	TSharedPtr<FVariantManagerDisplayNode> PinnedNode = Node.Pin();
+	if (PinnedNode.IsValid())
+	{
+		if (PinnedNode->GetType() == EVariantManagerNodeType::Spacer)
+		{
+			return FReply::Unhandled();
+		}
+	}
+
+	return STableRow::OnDragOver(MyGeometry, DragDropEvent);
+}
+
+const FName SVariantManagerMultiColumnPropertyTableRow::ColumnPath(TEXT("Path"));
+const FName SVariantManagerMultiColumnPropertyTableRow::ColumnProperty(TEXT("Property"));
+const FName SVariantManagerMultiColumnPropertyTableRow::ColumnValue(TEXT("Value"));
+
+void SVariantManagerMultiColumnPropertyTableRow::Construct(
+	const FArguments& InArgs, 
+	const TSharedRef<STableViewBase>& OwnerTableView,
+	const TSharedRef<FVariantManagerPropertyNode>& InNode)
+{
+	Node = InNode;
+	bool bIsSelectable = InNode->IsSelectable();
+
+	SMultiColumnTableRow::Construct(
+		STableRow::FArguments()
+		.Padding(0)
+		.Style(&FAppStyle::Get().GetWidgetStyle<FTableRowStyle>("SceneOutliner.TableViewRow"))
+		.OnDragDetected(this, &SVariantManagerMultiColumnPropertyTableRow::DragDetected)
+		.OnCanAcceptDrop(this, &SVariantManagerMultiColumnPropertyTableRow::CanAcceptDrop)
+		.OnAcceptDrop(this, &SVariantManagerMultiColumnPropertyTableRow::AcceptDrop)
+		.OnDragLeave(this, &SVariantManagerMultiColumnPropertyTableRow::DragLeave)
+		.ShowSelection(bIsSelectable),
+		OwnerTableView);
+}
+
+FReply SVariantManagerMultiColumnPropertyTableRow::DragDetected(const FGeometry& InGeometry, const FPointerEvent& InPointerEvent)
+{
+	if (TSharedPtr<FVariantManagerDisplayNode> PinnedNode = Node.Pin())
+	{
+		if (TSharedPtr<FVariantManager> VarMan = PinnedNode->GetVariantManager().Pin())
 		{
-			if (PinnedNode.IsValid() && PinnedNode->CanDrag())
+			const FSlateBrush* Icon = PinnedNode->GetIconBrush();
+
+			FVariantManagerSelection& Selection = VarMan->GetSelection();
+
+			TArray<FDisplayNodeRef> DraggableNodes;
+			FText DefaultHoverText;
+
+			// We'll drag a group of nodes based on what type we are (e.g. if we're variant or variant
+			// set, drag all of those)
+			switch (PinnedNode->GetType())
 			{
-				// We only support dragging one node property/function nodes for now
-				ensure(DraggableNodes.Num() == 0);
+			case EVariantManagerNodeType::Actor:
+			{
+				for (const FDisplayNodeRef SelectedNode : Selection.GetSelectedActorNodes())
+				{
+					if (SelectedNode->CanDrag())
+					{
+						DraggableNodes.Add(SelectedNode);
+					}
+				}
 
-				DraggableNodes.Add(PinnedNode.ToSharedRef());
-				DefaultHoverText = FText::Format(NSLOCTEXT("VariantManagerTableRow", "DragPropertyNode", "{0}"), PinnedNode->GetDisplayName());
+				int32 NumNodes = DraggableNodes.Num();
+				if (NumNodes > 1)
+				{
+					Icon = FSlateIconFinder::FindIconForClass(AActor::StaticClass()).GetOptionalIcon();
+				}
+
+				DefaultHoverText = FText::Format(NSLOCTEXT("VariantManagerTableRow", "DragActorNode", "{0} actor {0}|plural(one=binding,other=bindings)" ), NumNodes);
+				break;
 			}
-			break;
-		}
-		case EVariantManagerNodeType::Variant:  // Intended fallthrough
-		case EVariantManagerNodeType::VariantSet:
-		{
-			for (const FDisplayNodeRef& SelectedNode : Selection.GetSelectedOutlinerNodes())
+			case EVariantManagerNodeType::Function:
+			case EVariantManagerNodeType::Property:
 			{
-				if (SelectedNode->CanDrag())
+				if (PinnedNode.IsValid() && PinnedNode->CanDrag())
 				{
-					DraggableNodes.Add(SelectedNode);
+					// We only support dragging one node property/function nodes for now
+					ensure(DraggableNodes.Num() == 0);
+
+					DraggableNodes.Add(PinnedNode.ToSharedRef());
+					DefaultHoverText = FText::Format(NSLOCTEXT("VariantManagerTableRow", "DragPropertyNode", "{0}"), PinnedNode->GetDisplayName());
 				}
+				break;
 			}
+			case EVariantManagerNodeType::Variant:  // Intended fallthrough
+			case EVariantManagerNodeType::VariantSet:
+			{
+				for (const FDisplayNodeRef& SelectedNode : Selection.GetSelectedOutlinerNodes())
+				{
+					if (SelectedNode->CanDrag())
+					{
+						DraggableNodes.Add(SelectedNode);
+					}
+				}
 
-			int32 NumNodes = DraggableNodes.Num();
+				int32 NumNodes = DraggableNodes.Num();
 
-			DefaultHoverText = FText::Format(NSLOCTEXT("VariantManagerTableRow", "DragVariants", "{0} {0}|plural(one=variant,other=variants) and/or variant {0}|plural(one=set,other=sets)" ),
-				NumNodes);
+				DefaultHoverText = FText::Format(NSLOCTEXT("VariantManagerTableRow", "DragVariants", "{0} {0}|plural(one=variant,other=variants) and/or variant {0}|plural(one=set,other=sets)" ),
+					NumNodes);
 
-			break;
-		}
-		default:
-			break;
-		}
+				break;
+			}
+			default:
+				break;
+			}
 
-		if (DraggableNodes.Num() == 0)
-		{
-			return FReply::Unhandled();
-		}
+			if (DraggableNodes.Num() == 0)
+			{
+				return FReply::Unhandled();
+			}
 
-		VarMan->GetVariantManagerWidget()->SortDisplayNodes(DraggableNodes);
+			if (TSharedPtr<SVariantManager> VariantManagerWidget = VarMan->GetVariantManagerWidget())
+			{
+				VariantManagerWidget->SortDisplayNodes(DraggableNodes);
+			}
 
-		TSharedRef<FVariantManagerDragDropOp> DragDropOp = FVariantManagerDragDropOp::New(DraggableNodes);
-		DragDropOp->SetToolTip(DefaultHoverText, Icon);
-		DragDropOp->SetupDefaults();
+			TSharedRef<FVariantManagerDragDropOp> DragDropOp = FVariantManagerDragDropOp::New(DraggableNodes);
+			DragDropOp->SetToolTip(DefaultHoverText, Icon);
+			DragDropOp->SetupDefaults();
 
-		return FReply::Handled().BeginDragDrop(DragDropOp);
+			return FReply::Handled().BeginDragDrop(DragDropOp);
+		}
 	}
 
 	return FReply::Unhandled();
 }
 
-void SVariantManagerTableRow::DragLeave(const FDragDropEvent& DragDropEvent)
+void SVariantManagerMultiColumnPropertyTableRow::DragLeave(const FDragDropEvent& DragDropEvent)
 {
 	TSharedPtr<FDecoratedDragDropOp> DecoratedDragDropOp = DragDropEvent.GetOperationAs<FDecoratedDragDropOp>();
 	if (DecoratedDragDropOp.IsValid())
@@ -131,7 +309,10 @@ void SVariantManagerTableRow::DragLeave(const FDragDropEvent& DragDropEvent)
 	}
 }
 
-TOptional<EItemDropZone> SVariantManagerTableRow::CanAcceptDrop( const FDragDropEvent& DragDropEvent, EItemDropZone InItemDropZone, FDisplayNodeRef DisplayNode )
+TOptional<EItemDropZone> SVariantManagerMultiColumnPropertyTableRow::CanAcceptDrop(
+	const FDragDropEvent& DragDropEvent, 
+	EItemDropZone InItemDropZone,
+	TSharedPtr<FVariantManagerPropertyNode> DisplayNode)
 {
 	TSharedPtr<FVariantManagerDisplayNode> PinnedNode = Node.Pin();
 	if (PinnedNode.IsValid())
@@ -142,7 +323,7 @@ TOptional<EItemDropZone> SVariantManagerTableRow::CanAcceptDrop( const FDragDrop
 	return TOptional<EItemDropZone>();
 }
 
-FReply SVariantManagerTableRow::AcceptDrop( const FDragDropEvent& DragDropEvent, EItemDropZone InItemDropZone, FDisplayNodeRef DisplayNode )
+FReply SVariantManagerMultiColumnPropertyTableRow::AcceptDrop(const FDragDropEvent& DragDropEvent, EItemDropZone InItemDropZone, TSharedPtr<FVariantManagerPropertyNode> DisplayNode)
 {
 	TSharedPtr<FVariantManagerDisplayNode> PinnedNode = Node.Pin();
 	if (PinnedNode.IsValid())
@@ -154,15 +335,19 @@ FReply SVariantManagerTableRow::AcceptDrop( const FDragDropEvent& DragDropEvent,
 	return FReply::Unhandled();
 }
 
-FReply SVariantManagerTableRow::OnMouseButtonDoubleClick(const FGeometry& InMyGeometry, const FPointerEvent& InMouseEvent)
+FReply SVariantManagerMultiColumnPropertyTableRow::OnMouseButtonDoubleClick(const FGeometry& InMyGeometry, const FPointerEvent& InMouseEvent)
 {
 	STableRow::OnMouseButtonDoubleClick(InMyGeometry, InMouseEvent);
 
-	return Node.Pin()->OnDoubleClick(InMyGeometry, InMouseEvent);
+	if (TSharedPtr<FVariantManagerDisplayNode> PinnedNode = Node.Pin())
+	{
+		return PinnedNode->OnDoubleClick(InMyGeometry, InMouseEvent);
+	}
+
+	return FReply::Unhandled();
 }
 
-// Small hack to bypass CanDrop calls to spacer nodes, letting the underlying tree handle the events instead
-FReply SVariantManagerTableRow::OnDragOver(const FGeometry& MyGeometry, const FDragDropEvent& DragDropEvent)
+FReply SVariantManagerMultiColumnPropertyTableRow::OnDragOver(const FGeometry& MyGeometry, const FDragDropEvent& DragDropEvent)
 {
 	TSharedPtr<FVariantManagerDisplayNode> PinnedNode = Node.Pin();
 	if (PinnedNode.IsValid())
@@ -176,4 +361,61 @@ FReply SVariantManagerTableRow::OnDragOver(const FGeometry& MyGeometry, const FD
 	return STableRow::OnDragOver(MyGeometry, DragDropEvent);
 }
 
+TSharedRef<SWidget> SVariantManagerMultiColumnPropertyTableRow::GenerateWidgetForColumn(const FName& InColumnName)
+{
+	if (TSharedPtr<FVariantManagerPropertyNode> NodePtr = Node.Pin())
+	{
+		if (InColumnName == ColumnPath)
+		{
+			return SNew(SBox)
+				[
+					SNew(SHorizontalBox)
+					.Visibility(EVisibility::HitTestInvisible)
+
+					+ SHorizontalBox::Slot()
+					.Padding(2.0f, 0.0f, 2.0f, 0.0f)
+					.MaxWidth(15.0f)
+					.AutoWidth()
+					.VAlign(VAlign_Center)
+					.HAlign(HAlign_Center)
+					[
+						SNew(SImage)
+						.DesiredSizeOverride(FVector2D(16, 16))
+						.Image(FAppStyle::Get().GetBrush(NodePtr->IsActorProperty() ? "ClassIcon.Actor" : "ClassIcon.ActorComponent"))
+					]
+
+					+ SHorizontalBox::Slot()
+					.VAlign(VAlign_Center)
+					.Padding(4.0f, 0.0f, 2.0f, 0.0f)
+					[
+						SNew(STextBlock)
+						.Font(FAppStyle::GetFontStyle("Sequencer.AnimationOutliner.RegularFont"))
+						.Text(FText::FromString(NodePtr->GetSubPath()))
+					]
+				];
+		}
+		else if (InColumnName == ColumnProperty)
+		{
+			return SNew(SBox)
+				.VAlign(VAlign_Center)
+				.Visibility(EVisibility::HitTestInvisible)
+				[
+					SNew(STextBlock)
+					.Font(FAppStyle::GetFontStyle("Sequencer.AnimationOutliner.RegularFont"))
+					.Text(FText::FromString(NodePtr->GetPropertyName()))
+				];
+		}
+		else if (InColumnName == ColumnValue)
+		{
+			return SNew(SBox)
+				.VAlign(VAlign_Center)
+				[
+					NodePtr->GetCustomValueContent()
+				];
+		}
+	}
+
+	return SNullWidget::NullWidget;
+}
+
 #undef LOCTEXT_NAMESPACE
diff --git a/Engine/Plugins/Enterprise/VariantManager/Source/VariantManager/Private/SVariantManagerTableRow.h b/Engine/Plugins/Enterprise/VariantManager/Source/VariantManager/Private/SVariantManagerTableRow.h
--- a/Engine/Plugins/Enterprise/VariantManager/Source/VariantManager/Private/SVariantManagerTableRow.h
+++ b/Engine/Plugins/Enterprise/VariantManager/Source/VariantManager/Private/SVariantManagerTableRow.h
@@ -3,6 +3,7 @@
 #pragma once
 
 #include "DisplayNodes/VariantManagerDisplayNode.h"
+#include "DisplayNodes/VariantManagerPropertyNode.h"
 #include "Widgets/Views/SListView.h" // IWYU pragma: keep
 
 class STableViewBase;
@@ -38,3 +39,40 @@ private:
 
 	mutable TWeakPtr<FVariantManagerDisplayNode> Node;
 };
+
+
+/** Represents a multi-column row in the VariantManager's list views for displaying properties (and function callers) */
+class SVariantManagerMultiColumnPropertyTableRow : public SMultiColumnTableRow<TSharedPtr<FVariantManagerPropertyNode>>
+{
+public:
+	static const FName ColumnPath;
+	static const FName ColumnProperty;
+	static const FName ColumnValue;
+
+	SLATE_BEGIN_ARGS(SVariantManagerMultiColumnPropertyTableRow) {}
+	SLATE_END_ARGS()
+
+	void Construct(const FArguments& InArgs, const TSharedRef<STableViewBase>& OwnerTableView, const TSharedRef<FVariantManagerPropertyNode>& InNode);
+
+	TSharedPtr<FVariantManagerDisplayNode> GetDisplayNode() const
+	{
+		return Node.Pin();
+	}
+
+	// We subscribe these to base class STableRow's events
+	FReply DragDetected(const FGeometry& InGeometry, const FPointerEvent& InPointerEvent);
+	void DragLeave(const FDragDropEvent& DragDropEvent);
+	TOptional<EItemDropZone> CanAcceptDrop(const FDragDropEvent& DragDropEvent, EItemDropZone ItemDropZone, TSharedPtr<FVariantManagerPropertyNode> DisplayNode);
+	FReply AcceptDrop(const FDragDropEvent& DragDropEvent, EItemDropZone ItemDropZone, TSharedPtr<FVariantManagerPropertyNode> DisplayNode);
+
+	// STableRow interface
+	virtual FReply OnMouseButtonDoubleClick( const FGeometry& InMyGeometry, const FPointerEvent& InMouseEvent ) override;
+	virtual FReply OnDragOver(const FGeometry& MyGeometry, const FDragDropEvent& DragDropEvent) override;
+	//~ End STableRow interface
+
+	virtual TSharedRef<SWidget> GenerateWidgetForColumn(const FName& InColumnName) override;
+
+private:
+
+	mutable TWeakPtr<FVariantManagerPropertyNode> Node;
+};
diff --git a/Engine/Plugins/Enterprise/VariantManager/Source/VariantManager/Private/VariantManager.cpp b/Engine/Plugins/Enterprise/VariantManager/Source/VariantManager/Private/VariantManager.cpp
--- a/Engine/Plugins/Enterprise/VariantManager/Source/VariantManager/Private/VariantManager.cpp
+++ b/Engine/Plugins/Enterprise/VariantManager/Source/VariantManager/Private/VariantManager.cpp
@@ -81,7 +81,7 @@ UBlueprint* FVariantManager::GetOrCreateDirectorBlueprint(ULevelVariantSets* InL
 	return Director;
 }
 
-UK2Node_FunctionEntry* FVariantManager::CreateDirectorFunction(ULevelVariantSets* InLevelVariantSets, FName InFunctionName, UClass* PinClassType)
+UK2Node_FunctionEntry* FVariantManager::CreateDirectorFunction(ULevelVariantSets* InLevelVariantSets, FName InFunctionName, UClass* PinClassType, bool bCreateWithArguments)
 {
 	UBlueprint* Director = GetOrCreateDirectorBlueprint(InLevelVariantSets);
 
@@ -92,7 +92,7 @@ UK2Node_FunctionEntry* FVariantManager::CreateDirectorFunction(ULevelVariantSets
 
 	UEdGraph* Graph = FBlueprintEditorUtils::CreateNewGraph(Director, UniqueGraphName, UEdGraph::StaticClass(), UEdGraphSchema_K2::StaticClass());
 
-	const bool bIsUserCreated = false;
+	constexpr bool bIsUserCreated = false;
 	FBlueprintEditorUtils::AddFunctionGraph<UClass>(Director, Graph, bIsUserCreated, nullptr);
 
 	TArray<UK2Node_FunctionEntry*> EntryNodes;
@@ -118,6 +118,7 @@ UK2Node_FunctionEntry* FVariantManager::CreateDirectorFunction(ULevelVariantSets
 		PinType.PinSubCategoryObject = PinClassType ? PinClassType : UObject::StaticClass();
 		EntryNode->CreateUserDefinedPin(TARGET_PIN_NAME, PinType, EGPD_Output, true);
 
+		PinType.PinCategory = UEdGraphSchema_K2::PC_Object;
 		PinType.PinSubCategoryObject = ULevelVariantSets::StaticClass();
 		EntryNode->CreateUserDefinedPin(LEVEL_VARIANT_SETS_PIN_NAME, PinType, EGPD_Output, true);
 
@@ -127,6 +128,18 @@ UK2Node_FunctionEntry* FVariantManager::CreateDirectorFunction(ULevelVariantSets
 		PinType.PinSubCategoryObject = UVariant::StaticClass();
 		EntryNode->CreateUserDefinedPin(VARIANT_PIN_NAME, PinType, EGPD_Output, true);
 
+		if (bCreateWithArguments)
+		{
+			FEdGraphPinType ArgsPinType;
+			ArgsPinType.PinCategory = UEdGraphSchema_K2::PC_Name;
+			ArgsPinType.ContainerType = EPinContainerType::Map;
+			ArgsPinType.PinSubCategory = UEdGraphSchema_K2::PC_Name;
+			FEdGraphPinType ValuePinType;
+			ValuePinType.PinCategory = UEdGraphSchema_K2::PC_String;
+			ArgsPinType.PinValueType = FEdGraphTerminalType::FromPinType(ValuePinType);
+			EntryNode->CreateUserDefinedPin(ARGUMENTS_PIN_NAME, ArgsPinType, EGPD_Output, true);
+		}
+
 		EntryNode->ReconstructNode();
 
 		return EntryNode;
@@ -135,12 +148,12 @@ UK2Node_FunctionEntry* FVariantManager::CreateDirectorFunction(ULevelVariantSets
 	return nullptr;
 }
 
-UK2Node_FunctionEntry* FVariantManager::CreateDirectorFunctionFromFunction(UFunction* QuickBindFunction, UClass* PinClassType)
+UK2Node_FunctionEntry* FVariantManager::CreateDirectorFunctionFromFunction(UFunction* QuickBindFunction, UClass* PinClassType, bool bCreateWithArguments)
 {
 	FString DesiredNewEventName = FString(TEXT("Call ")) + QuickBindFunction->GetName();
 
 	// Create a single event binding and point all events in this property handle to it
-	UK2Node_FunctionEntry* NewFunctionEntry = CreateDirectorFunction(GetCurrentLevelVariantSets(), FName(*DesiredNewEventName), PinClassType);
+	UK2Node_FunctionEntry* NewFunctionEntry = CreateDirectorFunction(GetCurrentLevelVariantSets(), FName(*DesiredNewEventName), PinClassType, bCreateWithArguments);
 	if (!ensure(NewFunctionEntry))
 	{
 		return nullptr;
diff --git a/Engine/Plugins/Enterprise/VariantManager/Source/VariantManager/Private/VariantManagerBlueprintLibrary.cpp b/Engine/Plugins/Enterprise/VariantManager/Source/VariantManager/Private/VariantManagerBlueprintLibrary.cpp
--- a/Engine/Plugins/Enterprise/VariantManager/Source/VariantManager/Private/VariantManagerBlueprintLibrary.cpp
+++ b/Engine/Plugins/Enterprise/VariantManager/Source/VariantManager/Private/VariantManagerBlueprintLibrary.cpp
@@ -3,6 +3,8 @@
 #include "VariantManagerBlueprintLibrary.h"
 
 #include "CapturableProperty.h"
+#include "FunctionCaller.h"
+#include "K2Node_FunctionEntry.h"
 #include "LevelVariantSets.h"
 #include "LevelVariantSetsActor.h"
 #include "Trace/Trace.inl"
@@ -13,6 +15,8 @@
 #include "VariantManagerLog.h"
 #include "VariantObjectBinding.h"
 #include "VariantSet.h"
+#include "Kismet2/BlueprintEditorUtils.h"
+#include "Kismet2/KismetEditorUtilities.h"
 
 #include UE_INLINE_GENERATED_CPP_BY_NAME(VariantManagerBlueprintLibrary)
 
@@ -21,11 +25,18 @@ TUniquePtr<FVariantManager> UVariantManagerBlueprintLibrary::VariantManager;
 
 UVariantManagerBlueprintLibrary::UVariantManagerBlueprintLibrary(const FObjectInitializer& ObjectInitializer)
 	: Super(ObjectInitializer)
+{
+	GetVariantManager();
+}
+
+FVariantManager& UVariantManagerBlueprintLibrary::GetVariantManager()
 {
 	if (!VariantManager.IsValid())
 	{
 		VariantManager = MakeUnique<FVariantManager>();
 	}
+
+	return *VariantManager.Get();
 }
 
 ULevelVariantSets* UVariantManagerBlueprintLibrary::CreateLevelVariantSetsAsset(const FString& AssetName, const FString& AssetPath)
@@ -55,7 +66,7 @@ void UVariantManagerBlueprintLibrary::AddVariantSet(ULevelVariantSets* LevelVari
 		return;
 	}
 
-	VariantManager->AddVariantSets({VariantSet}, LevelVariantSets);
+	GetVariantManager().AddVariantSets({VariantSet}, LevelVariantSets);
 }
 
 void UVariantManagerBlueprintLibrary::AddVariant(UVariantSet* VariantSet, UVariant* Variant)
@@ -65,7 +76,7 @@ void UVariantManagerBlueprintLibrary::AddVariant(UVariantSet* VariantSet, UVaria
 		return;
 	}
 
-	VariantManager->AddVariants({Variant}, VariantSet);
+	GetVariantManager().AddVariants({Variant}, VariantSet);
 }
 
 void UVariantManagerBlueprintLibrary::AddActorBinding(UVariant* Variant, AActor* Actor)
@@ -75,7 +86,7 @@ void UVariantManagerBlueprintLibrary::AddActorBinding(UVariant* Variant, AActor*
 		return;
 	}
 
-	VariantManager->CreateObjectBindings({Actor}, {Variant});
+	GetVariantManager().CreateObjectBindings({Actor}, {Variant});
 }
 
 
@@ -88,7 +99,7 @@ void UVariantManagerBlueprintLibrary::RemoveVariantSet(ULevelVariantSets* LevelV
 		return;
 	}
 
-	VariantManager->RemoveVariantSetsFromParent({VariantSet});
+	GetVariantManager().RemoveVariantSetsFromParent({VariantSet});
 }
 
 void UVariantManagerBlueprintLibrary::RemoveVariant(UVariantSet* VariantSet, UVariant* Variant)
@@ -98,7 +109,7 @@ void UVariantManagerBlueprintLibrary::RemoveVariant(UVariantSet* VariantSet, UVa
 		return;
 	}
 
-	VariantManager->RemoveVariantsFromParent({Variant});
+	GetVariantManager().RemoveVariantsFromParent({Variant});
 }
 
 void UVariantManagerBlueprintLibrary::RemoveActorBinding(UVariant* Variant, AActor* Actor)
@@ -120,7 +131,10 @@ void UVariantManagerBlueprintLibrary::RemoveActorBinding(UVariant* Variant, AAct
 		}
 	}
 
-	VariantManager->RemoveObjectBindingsFromParent({TargetBinding});
+	if (TargetBinding)
+	{
+		GetVariantManager().RemoveObjectBindingsFromParent({TargetBinding});
+	}
 }
 
 
@@ -133,10 +147,9 @@ void UVariantManagerBlueprintLibrary::RemoveVariantSetByName(ULevelVariantSets*
 		return;
 	}
 
-	UVariantSet* VarSet = LevelVariantSets->GetVariantSetByName(VariantSetName);
-	if (VarSet)
+	if (UVariantSet* VarSet = LevelVariantSets->GetVariantSetByName(VariantSetName))
 	{
-		VariantManager->RemoveVariantSetsFromParent({VarSet});
+		GetVariantManager().RemoveVariantSetsFromParent({VarSet});
 	}
 }
 
@@ -147,10 +160,9 @@ void UVariantManagerBlueprintLibrary::RemoveVariantByName(UVariantSet* VariantSe
 		return;
 	}
 
-	UVariant* Var = VariantSet->GetVariantByName(VariantName);
-	if (Var)
+	if (UVariant* Var = VariantSet->GetVariantByName(VariantName))
 	{
-		VariantManager->RemoveVariantsFromParent({Var});
+		GetVariantManager().RemoveVariantsFromParent({Var});
 	}
 }
 
@@ -161,11 +173,9 @@ void UVariantManagerBlueprintLibrary::RemoveActorBindingByName(UVariant* Variant
 		return;
 	}
 
-	UVariantObjectBinding* Binding = Variant->GetBindingByName(ActorName);
-
-	if (Binding)
+	if (UVariantObjectBinding* Binding = Variant->GetBindingByName(ActorName))
 	{
-		VariantManager->RemoveObjectBindingsFromParent({Binding});
+		GetVariantManager().RemoveObjectBindingsFromParent({Binding});
 	}
 }
 
@@ -286,11 +296,11 @@ TArray<FString> UVariantManagerBlueprintLibrary::GetCapturableProperties(UObject
 
 	if (AActor* Actor = Cast<AActor>(ActorOrClass))
 	{
-		VariantManager->GetCapturableProperties({Actor}, OutProps, TargetPropertyPath, bCaptureAllArrayIndices);
+		GetVariantManager().GetCapturableProperties({Actor}, OutProps, TargetPropertyPath, bCaptureAllArrayIndices);
 	}
 	else if (UClass* Class = Cast<UClass>(ActorOrClass))
 	{
-		VariantManager->GetCapturableProperties({Class}, OutProps, TargetPropertyPath, bCaptureAllArrayIndices);
+		GetVariantManager().GetCapturableProperties({Class}, OutProps, TargetPropertyPath, bCaptureAllArrayIndices);
 	}
 
 	Result.Reserve(OutProps.Num());
@@ -342,7 +352,7 @@ UPropertyValue* UVariantManagerBlueprintLibrary::CaptureProperty(UVariant* Varia
 
 	const bool bCaptureAllArrayIndices = false;
 	TArray<TSharedPtr<FCapturableProperty>> OutProps;
-	VariantManager->GetCapturableProperties({Actor}, OutProps, PropertyPath, bCaptureAllArrayIndices);
+	GetVariantManager().GetCapturableProperties({Actor}, OutProps, PropertyPath, bCaptureAllArrayIndices);
 
 	if (OutProps.Num() < 1)
 	{
@@ -350,7 +360,7 @@ UPropertyValue* UVariantManagerBlueprintLibrary::CaptureProperty(UVariant* Varia
 		return nullptr;
 	}
 
-	TArray<UPropertyValue*> CreatedProps = VariantManager->CreatePropertyCaptures(OutProps, {TargetBinding});
+	TArray<UPropertyValue*> CreatedProps = GetVariantManager().CreatePropertyCaptures(OutProps, {TargetBinding});
 	if (CreatedProps.Num() > 0)
 	{
 		return CreatedProps[0];
@@ -385,6 +395,21 @@ void UVariantManagerBlueprintLibrary::DeleteDependency( UVariant* Variant, int32
 	}
 }
 
+TArray<FVariantDependency> UVariantManagerBlueprintLibrary::GetDependencies(UVariant* Variant)
+{
+	TArray<FVariantDependency> Dependencies;
+
+	if (Variant)
+	{
+		for (int Index = 0; Index < Variant->GetNumDependencies(); ++Index)
+		{
+			Dependencies.Add(Variant->GetDependency(Index));
+		}
+	}
+
+	return Dependencies;
+}
+
 TArray<UPropertyValue*> UVariantManagerBlueprintLibrary::GetCapturedProperties(UVariant* Variant, AActor* Actor)
 {
 	if (Variant == nullptr || Actor == nullptr)
@@ -418,7 +443,7 @@ void UVariantManagerBlueprintLibrary::RemoveCapturedProperty(UVariant* Variant,
 		return;
 	}
 
-	VariantManager->RemovePropertyCapturesFromParent({Property});
+	GetVariantManager().RemovePropertyCapturesFromParent({Property});
 }
 
 void UVariantManagerBlueprintLibrary::RemoveCapturedPropertyByName(UVariant* Variant, AActor* Actor, FString PropertyPath)
@@ -444,7 +469,7 @@ void UVariantManagerBlueprintLibrary::RemoveCapturedPropertyByName(UVariant* Var
 
 		if (FoundPropPtr)
 		{
-			VariantManager->RemovePropertyCapturesFromParent({*FoundPropPtr});
+			GetVariantManager().RemovePropertyCapturesFromParent({*FoundPropPtr});
 		}
 	}
 }
@@ -479,6 +504,12 @@ T GetPropertyValueImpl(UPropertyValue* Property, T InDefaultValue)
 		return InDefaultValue;
 	}
 
+	if (Property->GetValueSizeInBytes() < sizeof(T))
+	{
+		UE_LOG(LogVariantManager, Error, TEXT("Tried to access recorded data from property '%s', but size of property data is less then needed!"), *Property->GetFullDisplayString());
+		return InDefaultValue;
+	}
+
 	T Result;
 	FMemory::Memcpy(&Result, Property->GetRecordedData().GetData(), sizeof(T));
 	return Result;
@@ -504,6 +535,12 @@ T GetPropertyValueRefTypeImpl(UPropertyValue* Property, T InDefaultValue)
 		return InDefaultValue;
 	}
 
+	if (Property->GetValueSizeInBytes() != sizeof(T))
+	{
+		UE_LOG(LogVariantManager, Error, TEXT("Tried to access recorded data from property '%s' as a type with size %llu, but the recorded data has size %d!"), *Property->GetFullDisplayString(), sizeof(T), Property->GetValueSizeInBytes());
+		return InDefaultValue;
+	}
+
 	return *((const T*)Property->GetRecordedData().GetData());
 }
 
@@ -927,3 +964,670 @@ FIntPoint UVariantManagerBlueprintLibrary::GetValueIntPoint(UPropertyValue* Prop
 {
 	return GetPropertyValueImpl(Property, FIntPoint());
 }
+
+UBlueprint* UVariantManagerBlueprintLibrary::GetOrCreateDirectorBlueprint(UVariant* Variant)
+{
+	if (Variant)
+	{
+		if (UVariantSet* VariantSet = Variant->GetParent())
+		{
+			if (ULevelVariantSets* LevelVariantSets = VariantSet->GetParent())
+			{
+				GetVariantManager().InitVariantManager(LevelVariantSets);
+
+				return GetVariantManager().GetOrCreateDirectorBlueprint(LevelVariantSets);
+			}
+
+			UE_LOG(LogVariantManager, Error, TEXT("VariantSet '%s' does not have a parent LevelVariantSets!"), *VariantSet->GetDisplayText().ToString());
+		}
+		else
+		{
+			UE_LOG(LogVariantManager, Error, TEXT("Variant '%s' does not have a parent VariantSet!"), *Variant->GetDisplayText().ToString());
+		}
+	}
+	else
+	{
+		UE_LOG(LogVariantManager, Error, TEXT("Variant was null!"));
+	}
+	
+	return nullptr;
+}
+
+TArray<FName> UVariantManagerBlueprintLibrary::GetFunctionCallerNames(UVariant* Variant, AActor* Actor)
+{
+	if (Variant == nullptr || Actor == nullptr)
+	{
+		UE_LOG(LogVariantManager, Error, TEXT("Variant or Actor was null!"));
+
+		return {};
+	}
+
+	TArray<FName> FunctionCallerNames;
+
+	const FString ActorName = Actor->GetName();
+	if (UVariantObjectBinding* Binding = Variant->GetBindingByName(ActorName))
+	{
+		const TArray<FFunctionCaller>& FunctionCallers = Binding->GetFunctionCallers();
+
+		Algo::Transform(FunctionCallers, FunctionCallerNames, [](const FFunctionCaller& Caller)
+		{
+			return Caller.FunctionName;
+		});
+	}
+	else
+	{
+		UE_LOG(LogVariantManager, Warning, TEXT("Variant '%s' does not have a binding for actor '%s'!"), *Variant->GetDisplayText().ToString(), *Actor->GetActorLabel());
+	}
+
+	return FunctionCallerNames;
+}
+
+bool UVariantManagerBlueprintLibrary::GetFunctionCallerArguments(
+	UVariant* Variant, 
+	AActor* Actor, 
+	FName FunctionName,
+	TMap<FName, FString>& Arguments)
+{
+	Arguments.Empty();
+
+	if (Variant == nullptr || Actor == nullptr || FunctionName.IsNone() || !FunctionName.IsValid())
+	{
+		UE_LOG(LogVariantManager, Error, TEXT("Variant, Actor, or FunctionName was not specified!"));
+		return false;
+	}
+
+	const FString ActorName = Actor->GetName();
+	if (const UVariantObjectBinding* Binding = Variant->GetBindingByName(ActorName))
+	{
+		const TArray<FFunctionCaller>& FunctionCallers = Binding->GetFunctionCallers();
+		if (const FFunctionCaller* FoundCaller = FunctionCallers.FindByPredicate([&FunctionName](const FFunctionCaller& Caller)
+		{
+			return Caller.FunctionName == FunctionName;
+		}))
+		{
+			Arguments = FoundCaller->FunctionArguments;
+			return true;
+		}
+
+		UE_LOG(LogVariantManager, Warning, TEXT("Function caller '%s' not found in variant '%s' for actor '%s'!"), *FunctionName.ToString(), *Variant->GetDisplayText().ToString(), *Actor->GetActorLabel());
+	}
+	else
+	{
+		UE_LOG(LogVariantManager, Warning, TEXT("Variant '%s' does not have a binding for actor '%s'!"), *Variant->GetDisplayText().ToString(), *Actor->GetActorLabel());
+	}
+
+	return false;
+}
+
+bool UVariantManagerBlueprintLibrary::CreateFunctionCaller(
+	UVariant* Variant, 
+	AActor* Actor, 
+	FName FunctionName,
+	EVariantFunctionCallerSignature Signature)
+{
+	if (Variant == nullptr || Actor == nullptr || FunctionName.IsNone() || !FunctionName.IsValid())
+	{
+		UE_LOG(LogVariantManager, Error, TEXT("Variant, Actor, or FunctionName was not specified!"));
+		return false;
+	}
+	
+	const FString ActorName = Actor->GetName();
+	UVariantObjectBinding* Binding = Variant->GetBindingByName(ActorName);
+	if (Binding)
+	{
+		if (CheckTheFunctionSignature(Variant, Actor, FunctionName, Signature))
+		{
+			UE_LOG(LogVariantManager, Display, TEXT("Function caller '%s' with matching signature already exists in variant '%s' for actor '%s', forwarding to AddFunctionCaller..."), *FunctionName.ToString(), *Variant->GetDisplayText().ToString(), *Actor->GetActorLabel());
+
+			return AddFunctionCaller(Variant, Actor, FunctionName);
+		}
+	}
+	else
+	{
+		TArray<UVariantObjectBinding*> Bindings = GetVariantManager().CreateObjectBindings({Actor}, {Variant});
+		if (Bindings.Num() == 1)
+		{
+			Binding = Bindings[0];
+		}
+	}
+
+	if (Binding)
+	{
+		if (UVariantSet* VariantSet = Variant->GetParent())
+		{
+			if (ULevelVariantSets* LevelVariantSets = VariantSet->GetParent())
+			{
+				GetVariantManager().InitVariantManager(LevelVariantSets);
+
+				if (UBlueprint* Director = GetVariantManager().GetOrCreateDirectorBlueprint(LevelVariantSets))
+				{
+					Director->Modify();
+					if (UEdGraph* Graph = FBlueprintEditorUtils::CreateNewGraph(Director, FunctionName, UEdGraph::StaticClass(), UEdGraphSchema_K2::StaticClass()))
+					{
+						constexpr bool bIsUserCreated = true;
+						FBlueprintEditorUtils::AddFunctionGraph<UClass>(Director, Graph, bIsUserCreated, nullptr);
+
+						TArray<UK2Node_FunctionEntry*> EntryNodes;
+						Graph->GetNodesOfClass<UK2Node_FunctionEntry>(EntryNodes);
+						
+						if (ensure(EntryNodes.Num() == 1 && EntryNodes[0]))
+						{
+							UK2Node_FunctionEntry* EntryNode = EntryNodes[0];
+
+							int32 ExtraFunctionFlags = ( FUNC_BlueprintCallable | FUNC_BlueprintEvent | FUNC_Public );
+							EntryNode->AddExtraFlags(ExtraFunctionFlags);
+
+							EntryNode->bIsEditable = true;
+							EntryNode->MetaData.Category = NSLOCTEXT("VariantManagerLibrary", "DefaultCategory", "Director functions");
+							EntryNode->MetaData.bCallInEditor = false;
+
+							EntryNode->NodeComment = TEXT("Will be called by the Variant Manager when switching variants.\nThe target pin will receive a reference to the actor to which\nthis function caller is bound.\n\nYou can modify this target pin, either removing it\nor changing its type to another object reference type.\nIf you do, however, in order for a reference to the bound actor to be assigned\nto the pin, the actor must be of a type derived from the pin type.");
+							EntryNode->bCommentBubblePinned = true;
+							EntryNode->bCommentBubbleVisible = true;
+
+							switch (Signature)
+							{
+							case EVariantFunctionCallerSignature::Static:
+								{
+									EntryNode->AddExtraFlags(FUNC_BlueprintPure);
+									EntryNode->ReconstructNode();
+								}
+								break;
+							case EVariantFunctionCallerSignature::NoParameters:
+								{}
+								break;
+							case EVariantFunctionCallerSignature::OneParameter:
+								{
+									FEdGraphPinType PinType;
+									PinType.PinCategory = UEdGraphSchema_K2::PC_Object;
+									PinType.PinSubCategoryObject = UObject::StaticClass();
+									EntryNode->CreateUserDefinedPin(TARGET_PIN_NAME, PinType, EGPD_Output, true);
+
+									EntryNode->ReconstructNode();
+								}
+								break;
+							case EVariantFunctionCallerSignature::OneParameterWithArgs:
+								{
+									{
+										FEdGraphPinType PinType;
+										PinType.PinCategory = UEdGraphSchema_K2::PC_Object;
+										PinType.PinSubCategoryObject = UObject::StaticClass();
+										EntryNode->CreateUserDefinedPin(TARGET_PIN_NAME, PinType, EGPD_Output, true);
+									}
+									{
+										FEdGraphPinType PinType;
+										PinType.PinCategory = UEdGraphSchema_K2::PC_Name;
+										PinType.ContainerType = EPinContainerType::Map;
+										PinType.PinSubCategory = UEdGraphSchema_K2::PC_Name;
+										FEdGraphPinType ValuePinType;
+										ValuePinType.PinCategory = UEdGraphSchema_K2::PC_String;
+										PinType.PinValueType = FEdGraphTerminalType::FromPinType(ValuePinType);
+										EntryNode->CreateUserDefinedPin(ARGUMENTS_PIN_NAME, PinType, EGPD_Output, true);
+									}
+									EntryNode->ReconstructNode();
+								}
+								break;
+							case EVariantFunctionCallerSignature::FourParameters:
+								{
+									FEdGraphPinType PinType;
+									PinType.PinCategory = UEdGraphSchema_K2::PC_Object;
+									PinType.PinSubCategoryObject = UObject::StaticClass();
+									EntryNode->CreateUserDefinedPin(TARGET_PIN_NAME, PinType, EGPD_Output, true);
+
+									PinType.PinSubCategoryObject = ULevelVariantSets::StaticClass();
+									EntryNode->CreateUserDefinedPin(LEVEL_VARIANT_SETS_PIN_NAME, PinType, EGPD_Output, true);
+
+									PinType.PinSubCategoryObject = UVariantSet::StaticClass();
+									EntryNode->CreateUserDefinedPin(VARIANT_SET_PIN_NAME, PinType, EGPD_Output, true);
+
+									PinType.PinSubCategoryObject = UVariant::StaticClass();
+									EntryNode->CreateUserDefinedPin(VARIANT_PIN_NAME, PinType, EGPD_Output, true);
+
+									EntryNode->ReconstructNode();
+								}
+								break;
+							case EVariantFunctionCallerSignature::FourParametersWithArgs:
+								{
+									{
+										FEdGraphPinType PinType;
+										PinType.PinCategory = UEdGraphSchema_K2::PC_Object;
+										PinType.PinSubCategoryObject = UObject::StaticClass();
+										EntryNode->CreateUserDefinedPin(TARGET_PIN_NAME, PinType, EGPD_Output, true);
+
+										PinType.PinSubCategoryObject = ULevelVariantSets::StaticClass();
+										EntryNode->CreateUserDefinedPin(LEVEL_VARIANT_SETS_PIN_NAME, PinType, EGPD_Output, true);
+
+										PinType.PinSubCategoryObject = UVariantSet::StaticClass();
+										EntryNode->CreateUserDefinedPin(VARIANT_SET_PIN_NAME, PinType, EGPD_Output, true);
+
+										PinType.PinSubCategoryObject = UVariant::StaticClass();
+										EntryNode->CreateUserDefinedPin(VARIANT_PIN_NAME, PinType, EGPD_Output, true);
+									}
+									{
+										FEdGraphPinType PinType;
+										PinType.PinCategory = UEdGraphSchema_K2::PC_Name;
+										PinType.ContainerType = EPinContainerType::Map;
+										PinType.PinSubCategory = UEdGraphSchema_K2::PC_Name;
+										FEdGraphPinType ValuePinType;
+										ValuePinType.PinCategory = UEdGraphSchema_K2::PC_String;
+										PinType.PinValueType = FEdGraphTerminalType::FromPinType(ValuePinType);
+										EntryNode->CreateUserDefinedPin(ARGUMENTS_PIN_NAME, PinType, EGPD_Output, true);
+									}
+
+									EntryNode->ReconstructNode();
+								}
+								break;
+							}
+
+					        // Compile the blueprint
+							FKismetEditorUtilities::CompileBlueprint(Director, EBlueprintCompileOptions::None, nullptr);
+
+							return AddFunctionCaller(Variant, Actor, FunctionName);
+						}
+					}
+				}
+			}
+		}
+	}
+	else
+	{
+		UE_LOG(LogVariantManager, Error, TEXT("Failed to create a binding for actor '%s' in variant '%s' when trying to create a function caller!"), *Actor->GetActorLabel(), *Variant->GetDisplayText().ToString());
+	}
+
+	return false;
+}
+
+bool UVariantManagerBlueprintLibrary::AddFunctionCaller(
+	UVariant* Variant, 
+	AActor* Actor, 
+	FName FunctionName)
+{
+	if (Variant == nullptr || Actor == nullptr || FunctionName.IsNone() || !FunctionName.IsValid())
+	{
+		UE_LOG(LogVariantManager, Error, TEXT("Variant, Actor, or FunctionName was not specified!"));
+		return false;
+	}
+
+	const FString ActorName = Actor->GetName();
+	if (UVariantObjectBinding* Binding = Variant->GetBindingByName(ActorName))
+	{
+		TArray<FFunctionCaller>& FunctionCallers = Binding->GetFunctionCallers();
+		if (FFunctionCaller* FoundCaller = FunctionCallers.FindByPredicate([&FunctionName](const FFunctionCaller& Caller)
+		{
+			return Caller.FunctionName == FunctionName;
+		}))
+		{
+			UE_LOG(LogVariantManager, Warning, TEXT("Function caller '%s' already exists in variant '%s' for actor '%s'!"), *FunctionName.ToString(), *Variant->GetDisplayText().ToString(), *Actor->GetActorLabel());
+
+			return false;
+		}
+
+		if (UVariantSet* VariantSet = Variant->GetParent())
+		{
+			if (ULevelVariantSets* LevelVariantSets = VariantSet->GetParent())
+			{
+				if (UBlueprint* Blueprint = Cast<UBlueprint>(LevelVariantSets->GetDirectorGeneratedBlueprint()))
+				{
+					UK2Node_FunctionEntry* ThisFunc = nullptr;
+
+					TArray<UK2Node_FunctionEntry*> EntryNodes;
+					for (UEdGraph* FunctionGraph : Blueprint->FunctionGraphs)
+					{
+						if (FunctionGraph->GetName() == FunctionName)
+						{
+							EntryNodes.Reset();
+							FunctionGraph->GetNodesOfClass<UK2Node_FunctionEntry>(EntryNodes);
+							if (EntryNodes.Num() == 1)
+							{
+								if (UK2Node_FunctionEntry* EntryNode = EntryNodes[0])
+								{
+									if (FFunctionCaller::IsValidFunction(EntryNode))
+									{
+										ThisFunc = EntryNode;
+										break;
+									}
+								}
+							}
+						}
+					}
+					
+					if (ThisFunc)
+					{
+						FFunctionCaller NewCaller;
+						NewCaller.FunctionName = FunctionName;
+						NewCaller.SetFunctionEntry(ThisFunc);
+						Binding->Modify();
+						Binding->AddFunctionCallers({ NewCaller });
+
+						return true;
+					}
+
+					UE_LOG(LogVariantManager, Warning, TEXT("Could not find function entry node for function '%s' in the director blueprint for variant '%s'!"), *FunctionName.ToString(), *Variant->GetDisplayText().ToString());
+				}
+
+				UE_LOG(LogVariantManager, Warning, TEXT("Function '%s' does not match any supported signatures for function callers in variant '%s' for actor '%s'!"), *FunctionName.ToString(), *Variant->GetDisplayText().ToString(), *Actor->GetActorLabel());
+			}
+			else
+			{
+				UE_LOG(LogVariantManager, Error, TEXT("VariantSet '%s' does not have a parent LevelVariantSets!"), *VariantSet->GetDisplayText().ToString());
+			}
+		}
+		else
+		{
+			UE_LOG(LogVariantManager, Error, TEXT("Variant '%s' does not have a parent VariantSet!"), *Variant->GetDisplayText().ToString());
+		}
+	}
+	else
+	{
+		UE_LOG(LogVariantManager, Warning, TEXT("Variant '%s' does not have a binding for actor '%s'!"), *Variant->GetDisplayText().ToString(), *Actor->GetActorLabel());
+	}
+
+	return false;
+}
+
+bool UVariantManagerBlueprintLibrary::UpdateFunctionCallerArguments(
+	UVariant* Variant, 
+	AActor* Actor, 
+	FName FunctionName,
+	const TMap<FName, FString>& Arguments)
+{
+	if (Variant == nullptr || Actor == nullptr || FunctionName.IsNone() || !FunctionName.IsValid())
+	{
+		UE_LOG(LogVariantManager, Error, TEXT("Variant, Actor, or FunctionName was not specified!"));
+		return false;
+	}
+
+	const FString ActorName = Actor->GetName();
+	if (UVariantObjectBinding* Binding = Variant->GetBindingByName(ActorName))
+	{
+		TArray<FFunctionCaller>& FunctionCallers = Binding->GetFunctionCallers();
+		if (FFunctionCaller* FoundCaller = FunctionCallers.FindByPredicate([&FunctionName](const FFunctionCaller& Caller)
+		{
+			return Caller.FunctionName == FunctionName;
+		}))
+		{
+			Binding->Modify();
+			FoundCaller->FunctionArguments = Arguments;
+			return true;
+		}
+
+		UE_LOG(LogVariantManager, Warning, TEXT("Function caller '%s' not found in variant '%s' for actor '%s'!"), *FunctionName.ToString(), *Variant->GetDisplayText().ToString(), *Actor->GetActorLabel());
+	}
+	else
+	{
+		UE_LOG(LogVariantManager, Warning, TEXT("Variant '%s' does not have a binding for actor '%s'!"), *Variant->GetDisplayText().ToString(), *Actor->GetActorLabel());
+	}
+
+	return false;
+}
+
+bool UVariantManagerBlueprintLibrary::RemoveFunctionCaller(
+	UVariant* Variant, 
+	AActor* Actor, 
+	FName FunctionName)
+{
+	if (Variant == nullptr || Actor == nullptr || FunctionName.IsNone() || !FunctionName.IsValid())
+	{
+		UE_LOG(LogVariantManager, Error, TEXT("Variant, Actor, or FunctionName was not specified!"));
+		return false;
+	}
+
+	const FString ActorName = Actor->GetName();
+	if (UVariantObjectBinding* Binding = Variant->GetBindingByName(ActorName))
+	{
+		TArray<FFunctionCaller>& FunctionCallers = Binding->GetFunctionCallers();
+		if (FFunctionCaller* FoundCaller = FunctionCallers.FindByPredicate([&FunctionName](const FFunctionCaller& Caller)
+		{
+			return Caller.FunctionName == FunctionName;
+		}))
+		{
+			Binding->Modify();
+			Binding->RemoveFunctionCallers({ FoundCaller });
+			return true;
+		}
+
+		UE_LOG(LogVariantManager, Warning, TEXT("Function caller '%s' not found in variant '%s' for actor '%s'!"), *FunctionName.ToString(), *Variant->GetDisplayText().ToString(), *Actor->GetActorLabel());
+	}
+	else
+	{
+		UE_LOG(LogVariantManager, Warning, TEXT("Variant '%s' does not have a binding for actor '%s'!"), *Variant->GetDisplayText().ToString(), *Actor->GetActorLabel());
+	}
+
+	return false;
+}
+
+namespace
+{
+	bool IsObjectParam(TFieldIterator<FProperty>& It)
+	{
+		if (It)
+		{
+			FProperty* Property = *It;
+
+			// Check property flags to determine how it is used
+			if (Property->HasAnyPropertyFlags(CPF_Parm))
+			{
+				// Check for input/output/return value
+				if (!Property->HasAnyPropertyFlags(CPF_ReturnParm))
+				{
+					// This is an input or output parameter
+					if (Property->HasAnyPropertyFlags(CPF_OutParm) && !Property->HasAnyPropertyFlags(CPF_ConstParm))
+					{
+						// Non-const reference: Output parameter (or Input/Output in BP)
+						return false;
+					}
+							                
+					// By value or const reference: Input parameter
+					if (Property->IsA(FObjectProperty::StaticClass()))
+					{
+						if (FObjectProperty* ObjectProp = CastField<FObjectProperty>(Property))
+						{
+							return true;
+						}
+					}
+				}
+			}
+		}
+
+		return false;
+	}
+
+	bool IsArgsParam(TFieldIterator<FProperty>& It)
+	{
+		if (It)
+		{
+			FProperty* Property = *It;
+
+			// Check property flags to determine how it is used
+			if (Property->HasAnyPropertyFlags(CPF_Parm))
+			{
+				// Check for input/output/return value
+				if (!Property->HasAnyPropertyFlags(CPF_ReturnParm))
+				{
+					// This is an input or output parameter
+					if (Property->HasAnyPropertyFlags(CPF_OutParm) && !Property->HasAnyPropertyFlags(CPF_ConstParm))
+					{
+						// Non-const reference: Output parameter (or Input/Output in BP)
+						return false;
+					}
+
+					if (FMapProperty* MapProp = CastField<FMapProperty>(Property))
+					{
+						if (CastField<FNameProperty>(MapProp->KeyProp) && 
+							CastField<FStrProperty>(MapProp->ValueProp))
+						{
+							return true;
+						}
+					}
+				}
+			}
+		}
+
+		return false;
+	}
+
+	bool IsObjectParamOfClass(TFieldIterator<FProperty>& It, const UClass* Class)
+	{
+		if (It && Class)
+		{
+			FProperty* Property = *It;
+
+			// Check property flags to determine how it is used
+			if (Property->HasAnyPropertyFlags(CPF_Parm))
+			{
+				// Check for input/output/return value
+				if (!Property->HasAnyPropertyFlags(CPF_ReturnParm))
+				{
+					// This is an input or output parameter
+					if (Property->HasAnyPropertyFlags(CPF_OutParm) && !Property->HasAnyPropertyFlags(CPF_ConstParm))
+					{
+						// Non-const reference: Output parameter (or Input/Output in BP)
+						return false;
+					}
+							                
+					// By value or const reference: Input parameter
+					if (Property->IsA(FObjectProperty::StaticClass()))
+					{
+						if (FObjectProperty* ObjectProp = CastField<FObjectProperty>(Property))
+						{
+							if (ObjectProp->PropertyClass == Class)
+							{
+								return true;
+							}
+						}
+					}
+				}
+			}
+		}
+
+		return false;
+	}
+}
+
+bool UVariantManagerBlueprintLibrary::CheckTheFunctionSignature(UVariant* Variant, AActor* BoundActor, FName FunctionName, EVariantFunctionCallerSignature Signature)
+{
+	if (Variant == nullptr || FunctionName.IsNone() || !FunctionName.IsValid())
+	{
+		UE_LOG(LogVariantManager, Error, TEXT("Variant or FunctionName was not specified!"));
+		return false;
+	}
+
+	if (UVariantSet* VariantSet = Variant->GetParent())
+	{
+		if (ULevelVariantSets* LevelVariantSets = VariantSet->GetParent())
+		{
+			if (UObject* DirectorInstance = LevelVariantSets->GetDirectorGeneratedClass())
+			{
+				if (UFunction* Function = DirectorInstance->FindFunction(FunctionName))
+				{
+					switch (Signature)
+					{
+					case EVariantFunctionCallerSignature::Static:
+						if (Function->HasAnyFunctionFlags(FUNC_BlueprintPure) &&
+							Function->NumParms == 0)
+						{
+							return true;
+						}
+						break;
+					case EVariantFunctionCallerSignature::NoParameters:
+						if (Function->HasAnyFunctionFlags(FUNC_BlueprintCallable) &&
+							!Function->HasAnyFunctionFlags(FUNC_BlueprintPure) &&
+							Function->NumParms == 0)
+						{
+							return true;
+						}
+						break;
+					case EVariantFunctionCallerSignature::OneParameter:
+						if (Function->HasAnyFunctionFlags(FUNC_BlueprintCallable) &&
+							!Function->HasAnyFunctionFlags(FUNC_BlueprintPure) &&
+							Function->NumParms == 1)
+						{
+						    TFieldIterator<FProperty> It(Function);
+
+						    if (IsObjectParam(It))
+						    {
+							    return true;
+						    }
+
+						    return false;
+						}
+						break;
+					case EVariantFunctionCallerSignature::OneParameterWithArgs:
+						if (Function->HasAnyFunctionFlags(FUNC_BlueprintCallable) &&
+							!Function->HasAnyFunctionFlags(FUNC_BlueprintPure) &&
+							Function->NumParms == 2)
+						{
+						    TFieldIterator<FProperty> It(Function);
+							
+							if (IsObjectParam(It))
+						    {
+								++It;
+								if (IsArgsParam(It))
+								{
+									return true;
+								}
+							}
+						}
+						break;
+					case EVariantFunctionCallerSignature::FourParameters:
+						if (Function->HasAnyFunctionFlags(FUNC_BlueprintCallable) &&
+							!Function->HasAnyFunctionFlags(FUNC_BlueprintPure) &&
+							Function->NumParms == 4)
+						{
+						    TFieldIterator<FProperty> It(Function);
+							
+							if (IsObjectParam(It))
+						    {
+								++It;
+								if (IsObjectParamOfClass(It, ULevelVariantSets::StaticClass()))
+								{
+									++It;
+									if (IsObjectParamOfClass(It, UVariantSet::StaticClass()))
+									{
+										++It;
+										if (IsObjectParamOfClass(It, UVariant::StaticClass()))
+										{
+											return true;
+										}
+									}
+								}
+							}
+						}
+						break;
+					case EVariantFunctionCallerSignature::FourParametersWithArgs:
+						if (Function->HasAnyFunctionFlags(FUNC_BlueprintCallable) &&
+							!Function->HasAnyFunctionFlags(FUNC_BlueprintPure) &&
+							Function->NumParms == 5)
+						{
+						    TFieldIterator<FProperty> It(Function);
+							
+							if (IsObjectParam(It))
+						    {
+								++It;
+								if (IsObjectParamOfClass(It, ULevelVariantSets::StaticClass()))
+								{
+									++It;
+									if (IsObjectParamOfClass(It, UVariantSet::StaticClass()))
+									{
+										++It;
+										if (IsObjectParamOfClass(It, UVariant::StaticClass()))
+										{
+											++It;
+											if (IsArgsParam(It))
+											{
+												return true;
+											}
+										}
+									}
+								}
+							}
+						}
+						break;
+					}
+				}
+			}
+		}
+	}
+
+	return false;
+}
diff --git a/Engine/Plugins/Enterprise/VariantManager/Source/VariantManager/Private/VariantManagerEditorCommands.cpp b/Engine/Plugins/Enterprise/VariantManager/Source/VariantManager/Private/VariantManagerEditorCommands.cpp
--- a/Engine/Plugins/Enterprise/VariantManager/Source/VariantManager/Private/VariantManagerEditorCommands.cpp
+++ b/Engine/Plugins/Enterprise/VariantManager/Source/VariantManager/Private/VariantManagerEditorCommands.cpp
@@ -17,6 +17,7 @@ void FVariantManagerEditorCommands::RegisterCommands()
 	UI_COMMAND(AddSelectedActorsCommand, "Add selected actors", "Creates a new variant set at the bottom of the list", EUserInterfaceActionType::Button, FInputChord());
 	UI_COMMAND(SwitchOnSelectedVariantCommand, "Switch on", "Applies a single variant's values to all its captured properties", EUserInterfaceActionType::Button, FInputChord());
 	UI_COMMAND(CreateThumbnailCommand, "Set from viewport", "Sets the thumbnail to what is currently shown on the viewport", EUserInterfaceActionType::Button, FInputChord());
+	UI_COMMAND(SelectThumbnailAssetCommand, "Set from asset...", "Use an asset to serve as the thumbnail", EUserInterfaceActionType::Button, FInputChord());
 	UI_COMMAND(LoadThumbnailCommand, "Set from file", "Load an image to serve as the thumbnail", EUserInterfaceActionType::Button, FInputChord());
 	UI_COMMAND(ClearThumbnailCommand, "Clear", "Resets the thumbnail to default", EUserInterfaceActionType::Button, FInputChord());
 	UI_COMMAND(AddPropertyCaptures, "Add properties", "Capture new properties from the selected actors to this variant", EUserInterfaceActionType::Button, FInputChord());
diff --git a/Engine/Plugins/Enterprise/VariantManager/Source/VariantManager/Private/VariantManagerEditorCommands.h b/Engine/Plugins/Enterprise/VariantManager/Source/VariantManager/Private/VariantManagerEditorCommands.h
--- a/Engine/Plugins/Enterprise/VariantManager/Source/VariantManager/Private/VariantManagerEditorCommands.h
+++ b/Engine/Plugins/Enterprise/VariantManager/Source/VariantManager/Private/VariantManagerEditorCommands.h
@@ -21,6 +21,7 @@ public:
 	TSharedPtr<FUICommandInfo> AddSelectedActorsCommand;
 	TSharedPtr<FUICommandInfo> SwitchOnSelectedVariantCommand;
 	TSharedPtr<FUICommandInfo> CreateThumbnailCommand;
+	TSharedPtr<FUICommandInfo> SelectThumbnailAssetCommand;
 	TSharedPtr<FUICommandInfo> LoadThumbnailCommand;
 	TSharedPtr<FUICommandInfo> ClearThumbnailCommand;
 	TSharedPtr<FUICommandInfo> AddPropertyCaptures;
diff --git a/Engine/Plugins/Enterprise/VariantManager/Source/VariantManager/Private/VariantManagerModule.cpp b/Engine/Plugins/Enterprise/VariantManager/Source/VariantManager/Private/VariantManagerModule.cpp
--- a/Engine/Plugins/Enterprise/VariantManager/Source/VariantManager/Private/VariantManagerModule.cpp
+++ b/Engine/Plugins/Enterprise/VariantManager/Source/VariantManager/Private/VariantManagerModule.cpp
@@ -85,7 +85,7 @@ public:
 		TabManager->RegisterTabSpawner( FLevelVariantSetsEditorToolkit::GetVariantManagerTabID(), FOnSpawnTab::CreateStatic( &FVariantManagerModule::HandleTabManagerSpawnTab ) )
 			.SetDisplayName( LOCTEXT("VariantManagerMainTab", "Variant Manager") )
 			.SetGroup( WorkspaceMenu::GetMenuStructure().GetLevelEditorCategory() )
-			.SetIcon( FSlateIcon(VariantManagerStyleSetName, "VariantManager.Icon") );
+			.SetIcon( FSlateIcon(VariantManagerStyleSetName, "VariantManager.Icon.Small") );
 	}
 
 	static void UnregisterTabSpawner( const TSharedPtr< FTabManager >& TabManager )
diff --git a/Engine/Plugins/Enterprise/VariantManager/Source/VariantManager/Private/VariantManagerStyle.cpp b/Engine/Plugins/Enterprise/VariantManager/Source/VariantManager/Private/VariantManagerStyle.cpp
--- a/Engine/Plugins/Enterprise/VariantManager/Source/VariantManager/Private/VariantManagerStyle.cpp
+++ b/Engine/Plugins/Enterprise/VariantManager/Source/VariantManager/Private/VariantManagerStyle.cpp
@@ -10,13 +10,16 @@
 
 FVariantManagerStyle::FVariantManagerStyle() : FSlateStyleSet("VariantManagerEditorStyle")
 {
-	SetContentRoot(IPluginManager::Get().FindPlugin(TEXT("VariantManager"))->GetContentDir());
+	if (TSharedPtr<IPlugin> Plugin = IPluginManager::Get().FindPlugin(TEXT("VariantManager")))
+	{
+		SetContentRoot(Plugin->GetContentDir());
+	}
 
 	// Const icon sizes
 	const FVector2D Icon8x8(8.0f, 8.0f);
 	const FVector2D Icon16x16(16.0f, 16.0f);
 	const FVector2D Icon20x20(20.0f, 20.0f);
-	const FVector2D Icon40x40(40.0f, 40.0f);
+	const FVector2D Icon64x64(64.0f, 64.0f);
 
 	/** Color used for the background of the entire variant manager as well as the spacer border */
 	Set( "VariantManager.Panels.LightBackgroundColor", FLinearColor( FColor( 96, 96, 96, 255 ) ) );
@@ -33,11 +36,22 @@ FVariantManagerStyle::FVariantManagerStyle() : FSlateStyleSet("VariantManagerEdi
 	/** The amount to indent child nodes of the layout tree */
 	Set( "VariantManager.Spacings.IndentAmount", 10.0f );
 
-	Set( "VariantManager.Icon", new IMAGE_BRUSH_SVG("VariantManager", Icon16x16) );
+	Set( "VariantManager.Icon.Small", new IMAGE_BRUSH_SVG("VariantManager", Icon16x16) );
+	Set( "VariantManager.Icon", new IMAGE_BRUSH_SVG("VariantManager_64", Icon64x64) );
+	Set( "VariantManager.VariantSet", new IMAGE_BRUSH_SVG("VariantSet_64", Icon64x64) );
+	Set( "VariantManager.VariantSet.Small", new IMAGE_BRUSH_SVG("VariantSet_16", Icon16x16) );
+	Set( "VariantManager.Variant", new IMAGE_BRUSH_SVG("Variant_64", Icon64x64) );
+	Set( "VariantManager.Variant.Small", new IMAGE_BRUSH_SVG("Variant_16", Icon16x16) );
+
+	Set( "VariantManager.AddFunction", new IMAGE_BRUSH_SVG("AddFunction_16", Icon16x16) );
+	Set( "VariantManager.AddBinding", new IMAGE_BRUSH_SVG("AddBinding_16", Icon16x16) );
 
 	Set("VariantManager.AutoCapture.Icon", new IMAGE_BRUSH_SVG("AutoCapture", Icon20x20));
 
 	SetContentRoot(FPaths::EngineContentDir());
+	
+	Set("VariantManager.IconFavorite", new IMAGE_BRUSH_SVG("/Slate/Starship/Common/Favorite", Icon16x16));
+
 	const FCheckBoxStyle RadioButtonStyle = FCheckBoxStyle()
 		.SetUncheckedImage(IMAGE_BRUSH_SVG("/Slate/Starship/CoreWidgets/CheckBox/radio-off", Icon16x16))
 		.SetUncheckedHoveredImage(IMAGE_BRUSH_SVG("/Slate/Starship/CoreWidgets/CheckBox/radio-off", Icon16x16, FLinearColor(0.5f, 0.5f, 0.5f, 1.0f)))
@@ -49,3 +63,8 @@ FVariantManagerStyle::FVariantManagerStyle() : FSlateStyleSet("VariantManagerEdi
 
 	FSlateStyleRegistry::RegisterSlateStyle(*this);
 }
+
+const FName FVariantManagerStyle::GetAppStyleSetName()
+{
+	return FName("VariantManagerEditorStyle");
+}
diff --git a/Engine/Plugins/Enterprise/VariantManager/Source/VariantManager/Private/VariantManagerStyle.h b/Engine/Plugins/Enterprise/VariantManager/Source/VariantManager/Private/VariantManagerStyle.h
--- a/Engine/Plugins/Enterprise/VariantManager/Source/VariantManager/Private/VariantManagerStyle.h
+++ b/Engine/Plugins/Enterprise/VariantManager/Source/VariantManager/Private/VariantManagerStyle.h
@@ -15,6 +15,8 @@ public:
 		FSlateStyleRegistry::UnRegisterSlateStyle(*this);
 	}
 
+	static const FName GetAppStyleSetName();
+
 	static FVariantManagerStyle& Get()
 	{
 		static FVariantManagerStyle Inst;
diff --git a/Engine/Plugins/Enterprise/VariantManager/Source/VariantManager/Public/CapturableProperty.h b/Engine/Plugins/Enterprise/VariantManager/Source/VariantManager/Public/CapturableProperty.h
--- a/Engine/Plugins/Enterprise/VariantManager/Source/VariantManager/Public/CapturableProperty.h
+++ b/Engine/Plugins/Enterprise/VariantManager/Source/VariantManager/Public/CapturableProperty.h
@@ -16,15 +16,38 @@ struct FCapturableProperty
 {
 	GENERATED_USTRUCT_BODY()
 
-	FCapturableProperty(){};
+	FCapturableProperty()
+		: Checked(false)
+		, Favorite(false)
+		, CaptureType(EPropertyValueCategory::Undefined)
+	{
+	};
 	FCapturableProperty(const FString& InDisplayName, const FPropertyPath& InPropPath, const TArray<FString>& InComponentNames, bool InChecked, FName InPropertySetterName, EPropertyValueCategory InCaptureType)
 		: DisplayName(InDisplayName)
 		, Prop(InPropPath)
 		, Checked(InChecked)
+		, Favorite(false)
 		, CaptureType(InCaptureType)
 		, ComponentNames(InComponentNames)
 		, PropertySetterName(InPropertySetterName)
 	{
+		int32 Index = INDEX_NONE;
+		PropertyName = InDisplayName;
+
+		if (InDisplayName.FindLastChar(TEXT('/'), Index))
+		{
+			if (Index >= 1 && InDisplayName.Len() > Index)
+			{
+				SubPath = InDisplayName.Left(Index).TrimStartAndEnd();
+				PropertyName = InDisplayName.Right(InDisplayName.Len() - Index - 1);
+			}
+		}
+
+		if (SubPath.IsEmpty())
+		{
+			SubPath = TEXT("Actor");
+			bIsActor = true;
+		}
 	}
 
 	UPROPERTY(BlueprintReadOnly, VisibleAnywhere, Category="VariantManager")
@@ -32,9 +55,14 @@ struct FCapturableProperty
 
 	FPropertyPath Prop;
 	bool Checked;
+	bool Favorite;
+	bool bIsActor = false;
 	EPropertyValueCategory CaptureType;
-	TArray<FString> ComponentNames; // Same length as Prop, but will be empty for non component ObjectProperties
+	TArray<FString> ComponentNames; // Same length as Prop, but will be empty for non-component ObjectProperties
 	FName PropertySetterName;
+	
+	FString SubPath;		// Path without property name
+	FString PropertyName;	// Just property name
 
 	TArray<FCapturedPropSegment> ToCapturedPropSegmentArray()
 	{
diff --git a/Engine/Plugins/Enterprise/VariantManager/Source/VariantManager/Public/VariantManager.h b/Engine/Plugins/Enterprise/VariantManager/Source/VariantManager/Public/VariantManager.h
--- a/Engine/Plugins/Enterprise/VariantManager/Source/VariantManager/Public/VariantManager.h
+++ b/Engine/Plugins/Enterprise/VariantManager/Source/VariantManager/Public/VariantManager.h
@@ -43,8 +43,8 @@ public:
 	// Sets up the blueprint class deriving from the function director that we'll use.
 	// Do this here as it uses editor-only code
 	UBlueprint* GetOrCreateDirectorBlueprint(ULevelVariantSets* InLevelVariantSets);
-	UK2Node_FunctionEntry* CreateDirectorFunction(ULevelVariantSets* InLevelVariantSets, FName InFunctionName, UClass* PinClassType);
-	UK2Node_FunctionEntry* CreateDirectorFunctionFromFunction(UFunction* QuickBindFunction, UClass* PinClassType);
+	UK2Node_FunctionEntry* CreateDirectorFunction(ULevelVariantSets* InLevelVariantSets, FName InFunctionName, UClass* PinClassType, bool bCreateWithArguments);
+	UK2Node_FunctionEntry* CreateDirectorFunctionFromFunction(UFunction* QuickBindFunction, UClass* PinClassType, bool bCreateWithArguments);
 	void CreateFunctionCaller(const TArray<UVariantObjectBinding*>& Bindings);
 
 	// Adds existing items to existing containers
diff --git a/Engine/Plugins/Enterprise/VariantManager/Source/VariantManager/Public/VariantManagerBlueprintLibrary.h b/Engine/Plugins/Enterprise/VariantManager/Source/VariantManager/Public/VariantManagerBlueprintLibrary.h
--- a/Engine/Plugins/Enterprise/VariantManager/Source/VariantManager/Public/VariantManagerBlueprintLibrary.h
+++ b/Engine/Plugins/Enterprise/VariantManager/Source/VariantManager/Public/VariantManagerBlueprintLibrary.h
@@ -17,6 +17,17 @@ class UVariantSet;
 class UVariant;
 class UPropertyValue;
 
+UENUM(BlueprintType)
+enum class EVariantFunctionCallerSignature : uint8
+{
+	Static,
+	NoParameters,
+	OneParameter,
+	OneParameterWithArgs,
+	FourParameters,
+	FourParametersWithArgs
+};
+
 /**
 * Library of functions that can be used by the Python API to trigger VariantManager operations
 */
@@ -25,6 +36,9 @@ class UVariantManagerBlueprintLibrary : public UBlueprintFunctionLibrary
 {
 	GENERATED_UCLASS_BODY()
 
+	// Gets or initializes and returns Variant Manager
+	static FVariantManager& GetVariantManager();
+
 public:
 
 	// Creates a new LevelVariantSetsAsset named AssetName (e.g. 'MyLevelVariantSets') in the content path AssetPath (e.g. '/Game')
@@ -71,6 +85,8 @@ public:
 	UFUNCTION(BlueprintCallable, meta=(ScriptMethod), Category="VariantManager")
 	static void DeleteDependency(UVariant* Variant, int32 Index);
 
+	UFUNCTION(BlueprintCallable, meta=(ScriptMethod), Category="VariantManager")
+	static TArray<FVariantDependency> GetDependencies(UVariant* Variant);
 
 
 	// Get functions are the ones exposed directly on the VariantSet/Variant types
@@ -212,6 +228,30 @@ public:
 	UFUNCTION(BlueprintCallable, meta=(ScriptMethod), Category = "VariantManager|PropertyAccessors")
 	static FIntPoint GetValueIntPoint(UPropertyValue* Property);
 
+	UFUNCTION(BlueprintCallable, meta=(ScriptMethod), Category = "VariantManager|FunctionCallers")
+	static UBlueprint* GetOrCreateDirectorBlueprint(UVariant* Variant);
+
+	UFUNCTION(BlueprintCallable, meta=(ScriptMethod), Category = "VariantManager|FunctionCallers")
+	static TArray<FName> GetFunctionCallerNames(UVariant* Variant, AActor* Actor);
+
+	UFUNCTION(BlueprintCallable, meta=(ScriptMethod), Category = "VariantManager|FunctionCallers")
+	static bool GetFunctionCallerArguments(UVariant* Variant, AActor* Actor, FName FunctionName, TMap<FName, FString>& Arguments);
+
+	UFUNCTION(BlueprintCallable, meta=(ScriptMethod), Category = "VariantManager|FunctionCallers")
+	static bool CreateFunctionCaller(UVariant* Variant, AActor* Actor, FName FunctionName, EVariantFunctionCallerSignature Signature);
+
+	UFUNCTION(BlueprintCallable, meta=(ScriptMethod), Category = "VariantManager|FunctionCallers")
+	static bool AddFunctionCaller(UVariant* Variant, AActor* Actor, FName FunctionName);
+
+	UFUNCTION(BlueprintCallable, meta=(ScriptMethod), Category = "VariantManager|FunctionCallers")
+	static bool UpdateFunctionCallerArguments(UVariant* Variant, AActor* Actor, FName FunctionName, const TMap<FName, FString>& Arguments);
+
+	UFUNCTION(BlueprintCallable, meta=(ScriptMethod), Category = "VariantManager|FunctionCallers")
+	static bool RemoveFunctionCaller(UVariant* Variant, AActor* Actor, FName FunctionName);
+
+private:
+	static bool CheckTheFunctionSignature(UVariant* Variant, AActor* BoundActor, FName FunctionName, EVariantFunctionCallerSignature Signature);
+
 private:
 	static TUniquePtr<FVariantManager> VariantManager;
 };
diff --git a/Engine/Plugins/Enterprise/VariantManager/Source/VariantManager/VariantManager.Build.cs b/Engine/Plugins/Enterprise/VariantManager/Source/VariantManager/VariantManager.Build.cs
--- a/Engine/Plugins/Enterprise/VariantManager/Source/VariantManager/VariantManager.Build.cs
+++ b/Engine/Plugins/Enterprise/VariantManager/Source/VariantManager/VariantManager.Build.cs
@@ -24,14 +24,17 @@ namespace UnrealBuildTool.Rules
 				new string[]
 				{
 					"AppFramework", // For color pickers (for color and linear color properties)
+					"AssetTools",
                     "BlueprintGraph", // For function director
 					"DesktopPlatform",
 					"EditorStyle", // For FontAwesome glyphs
+					"EditorWidgets",
 					"GraphEditor", // For DragDropOp, might be removed later
 					"InputCore", // For ListView keyboard control
 					"Projects", // So that we can use the IPluginManager, required for our custom style
 					"PropertyEditor",  // For functions that create the property widgets
 					"CinematicCamera",  // So we can check the CineCamera structs exist
+					"ContentBrowser",
 					"SceneOutliner",
 					"Slate",
 					"SlateCore",
diff --git a/Engine/Plugins/Enterprise/VariantManagerContent/Source/VariantManagerContent/Private/FunctionCaller.cpp b/Engine/Plugins/Enterprise/VariantManagerContent/Source/VariantManagerContent/Private/FunctionCaller.cpp
--- a/Engine/Plugins/Enterprise/VariantManagerContent/Source/VariantManagerContent/Private/FunctionCaller.cpp
+++ b/Engine/Plugins/Enterprise/VariantManagerContent/Source/VariantManagerContent/Private/FunctionCaller.cpp
@@ -27,6 +27,94 @@ void FFunctionCaller::SetFunctionEntry(UK2Node_FunctionEntry* InFunctionEntry)
 	CacheFunctionName();
 }
 
+bool FFunctionCaller::CanHaveArguments() const
+{
+	if (IsBoundToBlueprint())
+	{
+		if (UK2Node_FunctionEntry* Function = GetFunctionEntry())
+		{
+			return CanHaveArguments(Function);
+		}
+	}
+
+	return false;
+}
+
+namespace
+{
+	bool IsTargetPinValid(const TSharedPtr<FUserPinInfo>& TargetPin)
+	{
+		return TargetPin.IsValid() &&
+			!TargetPin->PinType.bIsReference &&
+			(TargetPin->PinType.PinCategory == UEdGraphSchema_K2::PC_Object ||
+			TargetPin->PinType.PinCategory == UEdGraphSchema_K2::PC_Interface);
+	}
+
+	bool IsLevelVariantSetsPinValid(const TSharedPtr<FUserPinInfo>& TargetPin)
+	{
+		return TargetPin.IsValid() &&
+			!TargetPin->PinType.bIsReference &&
+			TargetPin->PinType.PinCategory == UEdGraphSchema_K2::PC_Object &&
+			TargetPin->PinType.PinSubCategoryObject == ULevelVariantSets::StaticClass();
+	}
+
+	bool IsVariantSetPinValid(const TSharedPtr<FUserPinInfo>& TargetPin)
+	{
+		return TargetPin.IsValid() &&
+			!TargetPin->PinType.bIsReference &&
+			TargetPin->PinType.PinCategory == UEdGraphSchema_K2::PC_Object &&
+			TargetPin->PinType.PinSubCategoryObject == UVariantSet::StaticClass();
+	}
+
+	bool IsVariantPinValid(const TSharedPtr<FUserPinInfo>& TargetPin)
+	{
+		return TargetPin.IsValid() &&
+			!TargetPin->PinType.bIsReference &&
+			TargetPin->PinType.PinCategory == UEdGraphSchema_K2::PC_Object &&
+			TargetPin->PinType.PinSubCategoryObject == UVariant::StaticClass();
+	}
+
+	bool IsArgumentsPinValid(const TSharedPtr<FUserPinInfo>& TargetPin)
+	{
+		return TargetPin.IsValid() &&
+			!TargetPin->PinType.bIsReference &&
+			TargetPin->PinType.IsMap() &&
+			TargetPin->PinType.PinCategory == UEdGraphSchema_K2::PC_Name &&
+			TargetPin->PinType.PinValueType.TerminalCategory == UEdGraphSchema_K2::PC_String;
+	}
+
+}
+
+bool FFunctionCaller::CanHaveArguments(UK2Node_FunctionEntry* Function)
+{
+	if (Function)
+	{
+		if (Function->UserDefinedPins.Num() == 2)
+		{
+			TSharedPtr<FUserPinInfo> TargetPin		= Function->UserDefinedPins[0];
+			TSharedPtr<FUserPinInfo> ArgumentsPin	= Function->UserDefinedPins[1];
+
+			return IsTargetPinValid(TargetPin) && IsArgumentsPinValid(ArgumentsPin);
+		}
+		else if (Function->UserDefinedPins.Num() == 5)
+		{
+			TSharedPtr<FUserPinInfo> TargetPin			 = Function->UserDefinedPins[0];
+			TSharedPtr<FUserPinInfo> LevelVariantSetsPin = Function->UserDefinedPins[1];
+			TSharedPtr<FUserPinInfo> VariantSetPin		 = Function->UserDefinedPins[2];
+			TSharedPtr<FUserPinInfo> VariantPin			 = Function->UserDefinedPins[3];
+			TSharedPtr<FUserPinInfo> ArgumentsPin		 = Function->UserDefinedPins[4];
+
+			return IsTargetPinValid(TargetPin) &&
+				IsLevelVariantSetsPinValid(LevelVariantSetsPin) &&
+				IsVariantSetPinValid(VariantSetPin) &&
+				IsVariantPinValid(VariantPin) &&
+				IsArgumentsPinValid(ArgumentsPin);
+		}
+	}
+
+	return false;
+}
+
 bool FFunctionCaller::IsBoundToBlueprint() const
 {
 	return IsValidFunction(GetFunctionEntry());
@@ -44,14 +132,14 @@ bool FFunctionCaller::IsValidFunction(UK2Node_FunctionEntry* Function)
 	}
 	else if (Function->UserDefinedPins.Num() == 1)
 	{
-		TSharedPtr<FUserPinInfo> TargetPin = Function->UserDefinedPins[0];
-		if (TargetPin->PinType.bIsReference ||
-			TargetPin->PinType.PinCategory != UEdGraphSchema_K2::PC_Object)
-		{
-			return false;
-		}
+		return IsTargetPinValid(Function->UserDefinedPins[0]);
+	}
+	else if (Function->UserDefinedPins.Num() == 2)
+	{
+		TSharedPtr<FUserPinInfo> TargetPin		= Function->UserDefinedPins[0];
+		TSharedPtr<FUserPinInfo> ArgumentsPin	= Function->UserDefinedPins[1];
 
-		return true;
+		return IsTargetPinValid(TargetPin) && IsArgumentsPinValid(ArgumentsPin);
 	}
 	else if (Function->UserDefinedPins.Num() == 4)
 	{
@@ -60,31 +148,24 @@ bool FFunctionCaller::IsValidFunction(UK2Node_FunctionEntry* Function)
 		TSharedPtr<FUserPinInfo> VariantSetPin		 = Function->UserDefinedPins[2];
 		TSharedPtr<FUserPinInfo> VariantPin			 = Function->UserDefinedPins[3];
 
-		if (TargetPin->PinType.bIsReference ||
-			TargetPin->PinType.PinCategory != UEdGraphSchema_K2::PC_Object)
-		{
-			return false;
-		}
-		if (LevelVariantSetsPin->PinType.bIsReference ||
-			LevelVariantSetsPin->PinType.PinCategory != UEdGraphSchema_K2::PC_Object ||
-			LevelVariantSetsPin->PinType.PinSubCategoryObject != ULevelVariantSets::StaticClass())
-		{
-			return false;
-		}
-		if (VariantSetPin->PinType.bIsReference ||
-			VariantSetPin->PinType.PinCategory != UEdGraphSchema_K2::PC_Object ||
-			VariantSetPin->PinType.PinSubCategoryObject != UVariantSet::StaticClass())
-		{
-			return false;
-		}
-		if (VariantPin->PinType.bIsReference ||
-			VariantPin->PinType.PinCategory != UEdGraphSchema_K2::PC_Object ||
-			VariantPin->PinType.PinSubCategoryObject != UVariant::StaticClass())
-		{
-			return false;
-		}
+		return IsTargetPinValid(TargetPin) &&
+			IsLevelVariantSetsPinValid(LevelVariantSetsPin) &&
+			IsVariantSetPinValid(VariantSetPin) &&
+			IsVariantPinValid(VariantPin);
+	}
+	else if (Function->UserDefinedPins.Num() == 5)
+	{
+		TSharedPtr<FUserPinInfo> TargetPin			 = Function->UserDefinedPins[0];
+		TSharedPtr<FUserPinInfo> LevelVariantSetsPin = Function->UserDefinedPins[1];
+		TSharedPtr<FUserPinInfo> VariantSetPin		 = Function->UserDefinedPins[2];
+		TSharedPtr<FUserPinInfo> VariantPin			 = Function->UserDefinedPins[3];
+		TSharedPtr<FUserPinInfo> ArgumentsPin		 = Function->UserDefinedPins[4];
 
-		return true;
+		return IsTargetPinValid(TargetPin) &&
+			IsLevelVariantSetsPinValid(LevelVariantSetsPin) &&
+			IsVariantSetPinValid(VariantSetPin) &&
+			IsVariantPinValid(VariantPin) &&
+			IsArgumentsPinValid(ArgumentsPin);
 	}
 
 	return false;
diff --git a/Engine/Plugins/Enterprise/VariantManagerContent/Source/VariantManagerContent/Private/LevelVariantSets.cpp b/Engine/Plugins/Enterprise/VariantManagerContent/Source/VariantManagerContent/Private/LevelVariantSets.cpp
--- a/Engine/Plugins/Enterprise/VariantManagerContent/Source/VariantManagerContent/Private/LevelVariantSets.cpp
+++ b/Engine/Plugins/Enterprise/VariantManagerContent/Source/VariantManagerContent/Private/LevelVariantSets.cpp
@@ -354,7 +354,7 @@ UVariantSet* ULevelVariantSets::GetVariantSetByName(FString VariantSetName)
 {
 	TObjectPtr<UVariantSet>* VarSetPtr = VariantSets.FindByPredicate([VariantSetName](const UVariantSet* VarSet)
 	{
-		return VarSet->GetDisplayText().ToString() == VariantSetName;
+		return VarSet && VarSet->GetDisplayText().ToString() == VariantSetName;
 	});
 
 	if (VarSetPtr)
@@ -364,6 +364,14 @@ UVariantSet* ULevelVariantSets::GetVariantSetByName(FString VariantSetName)
 	return nullptr;
 }
 
+int32 ULevelVariantSets::GetVariantSetIndex(FString VariantSetName) const
+{
+	return VariantSets.IndexOfByPredicate([VariantSetName](const UVariantSet* VarSet)
+	{
+		return VarSet && VarSet->GetDisplayText().ToString() == VariantSetName;
+	});
+}
+
 #if WITH_EDITOR
 void ULevelVariantSets::SetDirectorGeneratedBlueprint(UObject* InDirectorBlueprint)
 {
diff --git a/Engine/Plugins/Enterprise/VariantManagerContent/Source/VariantManagerContent/Private/PropertyValue.cpp b/Engine/Plugins/Enterprise/VariantManagerContent/Source/VariantManagerContent/Private/PropertyValue.cpp
--- a/Engine/Plugins/Enterprise/VariantManagerContent/Source/VariantManagerContent/Private/PropertyValue.cpp
+++ b/Engine/Plugins/Enterprise/VariantManagerContent/Source/VariantManagerContent/Private/PropertyValue.cpp
@@ -1102,6 +1102,19 @@ FString UPropertyValue::GetLeafDisplayString() const
 	return FullDisplayString;
 }
 
+FString UPropertyValue::GetBaseDisplayString() const
+{
+	FString LeftString;
+	FString RightString;
+
+	if(FullDisplayString.Split(PATH_DELIMITER, &LeftString, &RightString, ESearchCase::IgnoreCase, ESearchDir::FromEnd))
+	{
+		return LeftString;
+	}
+
+	return FullDisplayString;
+}
+
 int32 UPropertyValue::GetValueSizeInBytes() const
 {
 	FProperty* Prop = GetProperty();
diff --git a/Engine/Plugins/Enterprise/VariantManagerContent/Source/VariantManagerContent/Private/ThumbnailGenerator.cpp b/Engine/Plugins/Enterprise/VariantManagerContent/Source/VariantManagerContent/Private/ThumbnailGenerator.cpp
--- a/Engine/Plugins/Enterprise/VariantManagerContent/Source/VariantManagerContent/Private/ThumbnailGenerator.cpp
+++ b/Engine/Plugins/Enterprise/VariantManagerContent/Source/VariantManagerContent/Private/ThumbnailGenerator.cpp
@@ -398,15 +398,6 @@ UTexture2D* ThumbnailGenerator::GenerateThumbnailFromObjectThumbnail(UObject* Ob
 			UE_LOG(LogVariantContent, Warning, TEXT("Failed to create thumbnail texture from object '%s'!"), *Object->GetName());
 		}
 
-        if (UPackage* Package = Object->GetOutermost())
-        {
-            // After this our thumbnail will be empty, and we won't get in here ever again for this variant
-            ThumbnailTools::CacheEmptyThumbnail(Object->GetFullName(), Package);
-
-            // Updated the thumbnail in the package, so we need to notify that it changed
-            Package->MarkPackageDirty();
-        }
-
         return Thumbnail;
     }
 #endif
diff --git a/Engine/Plugins/Enterprise/VariantManagerContent/Source/VariantManagerContent/Private/Variant.cpp b/Engine/Plugins/Enterprise/VariantManagerContent/Source/VariantManagerContent/Private/Variant.cpp
--- a/Engine/Plugins/Enterprise/VariantManagerContent/Source/VariantManagerContent/Private/Variant.cpp
+++ b/Engine/Plugins/Enterprise/VariantManagerContent/Source/VariantManagerContent/Private/Variant.cpp
@@ -301,6 +301,21 @@ bool UVariant::IsActive()
 	return true;
 }
 
+void UVariant::SetThumbnailFromObject(UObject* ObjectThumbnail)
+{
+	if (ObjectThumbnail == nullptr)
+	{
+		SetThumbnailInternal(nullptr);
+	}
+	else
+	{
+		if (UTexture2D* NewThumbnail = ThumbnailGenerator::GenerateThumbnailFromObjectThumbnail(ObjectThumbnail))
+		{
+			SetThumbnailInternal(NewThumbnail);
+		}
+	}
+}
+
 void UVariant::SetThumbnailFromTexture(UTexture2D* Texture)
 {
 	if (Texture == nullptr)
diff --git a/Engine/Plugins/Enterprise/VariantManagerContent/Source/VariantManagerContent/Private/VariantObjectBinding.cpp b/Engine/Plugins/Enterprise/VariantManagerContent/Source/VariantManagerContent/Private/VariantObjectBinding.cpp
--- a/Engine/Plugins/Enterprise/VariantManagerContent/Source/VariantManagerContent/Private/VariantObjectBinding.cpp
+++ b/Engine/Plugins/Enterprise/VariantManagerContent/Source/VariantManagerContent/Private/VariantObjectBinding.cpp
@@ -1,9 +1,9 @@
-// Copyright Epic Games, Inc. All Rights Reserved.
+// Copyright Epic Games, Inc. All Rights Reserved. 
 
 #include "VariantObjectBinding.h"
-
-#include "LevelVariantSets.h"
-#include "PropertyValue.h"
+#include "UObject/UnrealType.h" 
+#include "LevelVariantSets.h" 
+#include "PropertyValue.h" 
 #include "Variant.h"
 #include "VariantManagerContentLog.h"
 #include "VariantManagerObjectVersion.h"
@@ -11,6 +11,8 @@
 
 #include "Engine/World.h"
 #include "FunctionCaller.h"
+#include "EdGraph/EdGraph.h"
+#include "Engine/Blueprint.h"
 
 #include UE_INLINE_GENERATED_CPP_BY_NAME(VariantObjectBinding)
 
@@ -248,6 +250,11 @@ TArray<FFunctionCaller>& UVariantObjectBinding::GetFunctionCallers()
 	return FunctionCallers;
 }
 
+const TArray<FFunctionCaller>& UVariantObjectBinding::GetFunctionCallers() const
+{
+	return FunctionCallers;
+}
+
 void UVariantObjectBinding::RemoveFunctionCallers(const TArray<FFunctionCaller*>& InFunctionCallers)
 {
 	Modify();
@@ -263,11 +270,121 @@ void UVariantObjectBinding::RemoveFunctionCallers(const TArray<FFunctionCaller*>
 #endif
 }
 
+namespace
+{
+	FObjectProperty* IsNonRefObjectParamNamed(FProperty* Prop, FName ExpectedName)
+	{ 
+		FObjectProperty* ObjProp = CastField<FObjectProperty>(Prop); 
+		if (!ObjProp)
+		{
+			return nullptr;
+		} 
+		if (ObjProp->GetFName() != ExpectedName)
+		{
+			return nullptr;
+		} 
+		if ((ObjProp->GetPropertyFlags() & CPF_ReferenceParm) != 0)
+		{
+			return nullptr;
+		} 
+		return ObjProp; 
+	}; 
+	
+	bool IsNonRefNameToStringMapParam(FProperty* Prop, FName ExpectedName)
+	{
+		FMapProperty* MapProp = CastField<FMapProperty>(Prop); 
+		if (!MapProp)
+		{
+			return false;
+		} 
+		if (MapProp->GetFName() != ExpectedName)
+		{
+			return false;
+		} 
+		if ((MapProp->GetPropertyFlags() & CPF_ReferenceParm) != 0)
+		{
+			return false;
+		} 
+		
+		const FNameProperty* KeyProp = CastField<FNameProperty>(MapProp->KeyProp);
+		const FStrProperty* ValueProp = CastField<FStrProperty>(MapProp->ValueProp);
+
+		return (KeyProp != nullptr) && (ValueProp != nullptr);
+	}; 
+		
+	bool ValidateTwoParamSignature(UFunction* InFunc, UClass*& OutExpectedTargetClass)
+	{
+		OutExpectedTargetClass = nullptr; 
+		FProperty* P0 = InFunc ? InFunc->PropertyLink : nullptr; 
+		FProperty* P1 = P0 ? P0->PropertyLinkNext : nullptr; 
+		FObjectProperty* TargetProp = IsNonRefObjectParamNamed(P0, TARGET_PIN_NAME); 
+		if (!TargetProp || !P1 || P1->PropertyLinkNext)
+		{
+			return false;
+		} 
+		if (!IsNonRefNameToStringMapParam(P1, TEXT("Arguments")))
+		{
+			return false;
+		} 
+		OutExpectedTargetClass = TargetProp->PropertyClass; 
+		return true;
+	}; 
+	
+	bool ValidateFiveParamSignature(UFunction* InFunc, UClass*& OutExpectedTargetClass)
+	{
+		OutExpectedTargetClass = nullptr;
+		FProperty* P0 = InFunc ? InFunc->PropertyLink : nullptr;
+		FProperty* P1 = P0 ? P0->PropertyLinkNext : nullptr;
+		FProperty* P2 = P1 ? P1->PropertyLinkNext : nullptr;
+		FProperty* P3 = P2 ? P2->PropertyLinkNext : nullptr;
+		FProperty* P4 = P3 ? P3->PropertyLinkNext : nullptr;
+		FObjectProperty* TargetProp = IsNonRefObjectParamNamed(P0, TARGET_PIN_NAME);
+		const FObjectProperty* LVSProp = CastField<FObjectProperty>(P1);
+		const FObjectProperty* VSProp = CastField<FObjectProperty>(P2);
+		const FObjectProperty* VProp = CastField<FObjectProperty>(P3);
+		if (!TargetProp || !LVSProp || !VSProp || !VProp || !P4 || P4->PropertyLinkNext)
+		{
+			return false;
+		} 
+		if ((LVSProp->GetPropertyFlags() & CPF_ReferenceParm) != 0 || 
+			(LVSProp->PropertyClass != ULevelVariantSets::StaticClass()))
+		{
+			return false;
+		} 
+		if ((VSProp->GetPropertyFlags() & CPF_ReferenceParm) != 0 || 
+			(VSProp->PropertyClass != UVariantSet::StaticClass()))
+		{
+			return false;
+		}
+		if ((VProp->GetPropertyFlags() & CPF_ReferenceParm) != 0 || 
+			(VProp->PropertyClass != UVariant::StaticClass()))
+		{
+			return false;
+		}
+		if (!IsNonRefNameToStringMapParam(P4, TEXT("Arguments")))
+		{
+			return false;
+		}
+		OutExpectedTargetClass = TargetProp->PropertyClass; 
+		return true;
+	};
+
+}
+
 void UVariantObjectBinding::ExecuteTargetFunction(FName FunctionName)
 {
 	ULevelVariantSets* ParentLVS = GetTypedOuter<ULevelVariantSets>();
+	if (!ParentLVS)
+	{
+		return;
+	}
 
 	UObject* BoundObject = GetObject();
+	if (!BoundObject)
+	{
+		return;
+	}
+
 	UObject* DirectorInstance = ParentLVS->GetDirectorInstance(BoundObject);
 	if (!DirectorInstance)
 	{
@@ -285,6 +402,13 @@ void UVariantObjectBinding::ExecuteTargetFunction(FName FunctionName)
 	const static FName NAME_CallInEditor(TEXT("CallInEditor"));
 
 	UWorld* World = DirectorInstance->GetWorld();
+
+	if (World == nullptr)
+	{
+		UE_LOG(LogVariantContent, Warning, TEXT("Cannot call function '%s' as Director Instance's World is null."), *FunctionName.ToString());
+		return;
+	}
+
 	if (World->WorldType == EWorldType::Editor && !Func->HasMetaData(NAME_CallInEditor))
 	{
 		UE_LOG(LogVariantContent, Warning, TEXT("Cannot call function '%s' as it doesn't have the CallInEditor option checked! Also note that calling this from the editor may have irreversible effects on the level."), *FunctionName.ToString());
@@ -316,6 +440,35 @@ void UVariantObjectBinding::ExecuteTargetFunction(FName FunctionName)
 			}
 		}
 	}
+	else if (Func->NumParms == 2)
+	{
+		UClass* ExpectedClass = nullptr;
+		if (!ValidateTwoParamSignature(Func, ExpectedClass))
+		{
+			UE_LOG(LogVariantContent, Error, TEXT("Function '%s' has NumParms==2 but does not match expected signature (Target, Arguments map)."), *Func->GetName());
+			return;
+		} 
+
+		if (ExpectedClass && !BoundObject->IsA(ExpectedClass))
+		{
+			UE_LOG(LogVariantContent, Error, TEXT("Failed to call function '%s' with object '%s' because it is not the correct type. Function expects a '%s' but target object is a '%s'."), *Func->GetName(), *BoundObject->GetName(), *ExpectedClass->GetName(), *BoundObject->GetClass()->GetName());
+			return;
+		} 
+
+		struct
+		{
+			UObject* Target = nullptr;
+			TMap<FName, FString> Arguments;
+		} FuncParams;
+		FuncParams.Target = BoundObject;
+
+		if (const FFunctionCaller* FunctionCaller = FunctionCallers.FindByPredicate([FunctionName](const FFunctionCaller& FunctionCaller) { return FunctionCaller.FunctionName == FunctionName; }))
+		{
+			FuncParams.Arguments = FunctionCaller->FunctionArguments;
+		}
+		
+		DirectorInstance->ProcessEvent(Func, &FuncParams);
+	}
 	else if (Func->NumParms == 4)
 	{
 		// Setup parameters to the Func as a struct
@@ -367,6 +520,42 @@ void UVariantObjectBinding::ExecuteTargetFunction(FName FunctionName)
 			);
 		}
 	}
+	else if (Func->NumParms == 5)
+	{
+		UClass* ExpectedClass = nullptr;
+		if (!ValidateFiveParamSignature(Func, ExpectedClass))
+		{
+			UE_LOG(LogVariantContent, Error, TEXT("Function '%s' has NumParms==5 but does not match expected signature (Target, LVS, VariantSet, Variant, Arguments map)."), *Func->GetName());
+			return;
+		}
+		if (ExpectedClass && !BoundObject->IsA(ExpectedClass))
+		{
+			UE_LOG(LogVariantContent, Error, TEXT("Failed to call function '%s' with object '%s' because it is not the correct type. Function expects a '%s' but target object is a '%s'."), 
+				*Func->GetName(), 
+				*BoundObject->GetName(),
+				*ExpectedClass->GetName(),
+				*BoundObject->GetClass()->GetName());
+			return;
+		} 
+		// Setup parameters to the Func as a struct
+		struct
+		{
+			UObject* Target = nullptr;
+			ULevelVariantSets* LevelVariantSets = nullptr;
+			UVariantSet* VariantSet = nullptr;
+			UVariant* Variant = nullptr;
+			TMap<FName, FString> Arguments;
+		} FuncParams;
+		FuncParams.Target = BoundObject;
+		FuncParams.Variant = GetParent();
+		FuncParams.VariantSet = FuncParams.Variant ? FuncParams.Variant->GetParent() : nullptr;
+		FuncParams.LevelVariantSets = FuncParams.VariantSet ? FuncParams.VariantSet->GetParent() : nullptr;
+		if (const FFunctionCaller* FunctionCaller = FunctionCallers.FindByPredicate([FunctionName](const FFunctionCaller& FunctionCaller) { return FunctionCaller.FunctionName == FunctionName; }))
+		{
+			FuncParams.Arguments = FunctionCaller->FunctionArguments;
+		}
+		DirectorInstance->ProcessEvent(Func, &FuncParams);
+	}
 }
 
 void UVariantObjectBinding::ExecuteAllTargetFunctions()
diff --git a/Engine/Plugins/Enterprise/VariantManagerContent/Source/VariantManagerContent/Private/VariantSet.cpp b/Engine/Plugins/Enterprise/VariantManagerContent/Source/VariantManagerContent/Private/VariantSet.cpp
--- a/Engine/Plugins/Enterprise/VariantManagerContent/Source/VariantManagerContent/Private/VariantSet.cpp
+++ b/Engine/Plugins/Enterprise/VariantManagerContent/Source/VariantManagerContent/Private/VariantSet.cpp
@@ -286,7 +286,7 @@ UVariant* UVariantSet::GetVariantByName(FString VariantName)
 {
 	TObjectPtr<UVariant>* VarPtr = Variants.FindByPredicate([VariantName](const UVariant* Var)
 	{
-		return Var->GetDisplayText().ToString() == VariantName;
+		return Var && Var->GetDisplayText().ToString() == VariantName;
 	});
 
 	if (VarPtr)
@@ -296,6 +296,29 @@ UVariant* UVariantSet::GetVariantByName(FString VariantName)
 	return nullptr;
 }
 
+int32 UVariantSet::GetVariantIndexByName(FString VariantName) const
+{
+	return Variants.IndexOfByPredicate([VariantName](const UVariant* Var)
+	{
+		return Var && Var->GetDisplayText().ToString() == VariantName;
+	});
+}
+
+void UVariantSet::SetThumbnailFromObject(UObject* ObjectThumbnail)
+{
+	if (ObjectThumbnail == nullptr)
+	{
+		SetThumbnailInternal(nullptr);
+	}
+	else
+	{
+		if (UTexture2D* NewThumbnail = ThumbnailGenerator::GenerateThumbnailFromObjectThumbnail(ObjectThumbnail))
+		{
+			SetThumbnailInternal(NewThumbnail);
+		}
+	}
+}
+
 void UVariantSet::SetThumbnailFromTexture(UTexture2D* Texture)
 {
 	if (Texture == nullptr)
@@ -343,6 +366,7 @@ UTexture2D* UVariantSet::GetThumbnail()
 void UVariantSet::SetThumbnailInternal(UTexture2D* NewThumbnail)
 {
 	Modify();
+
 	Thumbnail = NewThumbnail;
 
 	if (NewThumbnail)
diff --git a/Engine/Plugins/Enterprise/VariantManagerContent/Source/VariantManagerContent/Public/FunctionCaller.h b/Engine/Plugins/Enterprise/VariantManagerContent/Source/VariantManagerContent/Public/FunctionCaller.h
--- a/Engine/Plugins/Enterprise/VariantManagerContent/Source/VariantManagerContent/Public/FunctionCaller.h
+++ b/Engine/Plugins/Enterprise/VariantManagerContent/Source/VariantManagerContent/Public/FunctionCaller.h
@@ -10,6 +10,7 @@
 #define LEVEL_VARIANT_SETS_PIN_NAME TEXT("Level Variant Sets")
 #define VARIANT_SET_PIN_NAME TEXT("Variant Set")
 #define VARIANT_PIN_NAME TEXT("Variant")
+#define ARGUMENTS_PIN_NAME TEXT("Arguments")
 
 // This file is based on MovieSceneEvent.h
 
@@ -29,10 +30,21 @@ struct FFunctionCaller
 
 	/**
 	 * The function that should be called to invoke this event.
-	 * Functions must have either no parameters, or a single, pass-by-value object/interface parameter, with no return parameter.
+	 * Functions must have either:
+	 * - no parameters
+	 * - a single, pass-by-value object/interface parameter, with no return parameter
+	 * - two parameters, one for bound object and another for arguments
+	 * - four parameters: bound object, level variant set, variant set and variant
+	 * - five parameters: bound object, level variant set, variant set, variant and arguments
 	 */
 	UPROPERTY()
 	FName FunctionName;
+		
+	/**
+	 * The function arguments to be passed when invoking this event (provided function supports arguments).
+	 */
+	UPROPERTY()
+	TMap<FName, FString> FunctionArguments;
 
 #if WITH_EDITORONLY_DATA
 
@@ -63,7 +75,7 @@ public:
 	 * Helper function to determine whether the specified function entry is valid for this event
 	 *
 	 * @param Node         The node to test
-	 * @return true if the function entry node is compatible with a moviescene event, false otherwise
+	 * @return true if the function entry node is compatible with a variant director specification, false otherwise
 	 */
 	static UE_API bool IsValidFunction(UK2Node_FunctionEntry* Node);
 
@@ -83,6 +95,21 @@ public:
 	 */
 	UE_API void SetFunctionEntry(UK2Node_FunctionEntry* Entry);
 
+	/**
+	 * Returns true if bound to a function that can take arguments
+	 *
+	 * @return true if bound and if function can take arguments.
+	 */
+	UE_API bool CanHaveArguments() const;
+	
+	/**
+	 * Returns true if bound to a function that can take arguments
+	 * 
+	 * @param ThisFunc		The graph node to bind to
+	 * @return true if bound and if function can take arguments.
+	 */
+	UE_API static bool CanHaveArguments(UK2Node_FunctionEntry* ThisFunc);
+
 private:
 	/** Weak pointer to the function entry within the blueprint graph for this event. Stored as an editor-only UObject so UHT can parse it when building for non-editor. */
 	UPROPERTY()
diff --git a/Engine/Plugins/Enterprise/VariantManagerContent/Source/VariantManagerContent/Public/LevelVariantSets.h b/Engine/Plugins/Enterprise/VariantManagerContent/Source/VariantManagerContent/Public/LevelVariantSets.h
--- a/Engine/Plugins/Enterprise/VariantManagerContent/Source/VariantManagerContent/Public/LevelVariantSets.h
+++ b/Engine/Plugins/Enterprise/VariantManagerContent/Source/VariantManagerContent/Public/LevelVariantSets.h
@@ -43,6 +43,9 @@ public:
 
 	UFUNCTION(BlueprintPure, Category="LevelVariantSets")
 	UE_API UVariantSet* GetVariantSetByName(FString VariantSetName);
+	
+	UFUNCTION(BlueprintPure, Category="LevelVariantSets")
+	UE_API int32 GetVariantSetIndex(FString VariantSetName) const;
 
 #if WITH_EDITOR
 	UE_API void SetDirectorGeneratedBlueprint(UObject* InDirectorBlueprint);
diff --git a/Engine/Plugins/Enterprise/VariantManagerContent/Source/VariantManagerContent/Public/PropertyValue.h b/Engine/Plugins/Enterprise/VariantManagerContent/Source/VariantManagerContent/Public/PropertyValue.h
--- a/Engine/Plugins/Enterprise/VariantManagerContent/Source/VariantManagerContent/Public/PropertyValue.h
+++ b/Engine/Plugins/Enterprise/VariantManagerContent/Source/VariantManagerContent/Public/PropertyValue.h
@@ -129,6 +129,7 @@ public:
 	UFUNCTION(BlueprintCallable, Category="PropertyValue")
 	UE_API const FString& GetFullDisplayString() const;
 	UE_API FString GetLeafDisplayString() const;
+	UE_API FString GetBaseDisplayString() const;
 	UE_API virtual int32 GetValueSizeInBytes() const;
 	UE_API int32 GetPropertyOffsetInBytes() const;
 
diff --git a/Engine/Plugins/Enterprise/VariantManagerContent/Source/VariantManagerContent/Public/Variant.h b/Engine/Plugins/Enterprise/VariantManagerContent/Source/VariantManagerContent/Public/Variant.h
--- a/Engine/Plugins/Enterprise/VariantManagerContent/Source/VariantManagerContent/Public/Variant.h
+++ b/Engine/Plugins/Enterprise/VariantManagerContent/Source/VariantManagerContent/Public/Variant.h
@@ -81,6 +81,9 @@ public:
 	UFUNCTION(BlueprintCallable, Category="Variant")
 	UE_API bool IsActive();
 
+	UFUNCTION(BlueprintCallable, Category="Variant|Thumbnail")
+	UE_API void SetThumbnailFromObject(UObject* ObjectThumbnail);
+
 	// Sets the thumbnail to use for this variant. Can receive nullptr to clear it
 	UFUNCTION(BlueprintCallable, Category="Variant|Thumbnail")
 	UE_API void SetThumbnailFromTexture(UTexture2D* NewThumbnail);
@@ -122,6 +125,7 @@ public:
 
 	UFUNCTION(BlueprintCallable, BlueprintPure, Category = "Variant|Dependencies")
 	UE_API int32 GetNumDependencies();
+
 private:
 	UE_API void SetThumbnailInternal(UTexture2D* NewThumbnail);
 
diff --git a/Engine/Plugins/Enterprise/VariantManagerContent/Source/VariantManagerContent/Public/VariantObjectBinding.h b/Engine/Plugins/Enterprise/VariantManagerContent/Source/VariantManagerContent/Public/VariantObjectBinding.h
--- a/Engine/Plugins/Enterprise/VariantManagerContent/Source/VariantManagerContent/Public/VariantObjectBinding.h
+++ b/Engine/Plugins/Enterprise/VariantManagerContent/Source/VariantManagerContent/Public/VariantObjectBinding.h
@@ -37,6 +37,7 @@ public:
 
 	UE_API void AddFunctionCallers(const TArray<FFunctionCaller>& InFunctionCallers);
 	UE_API TArray<FFunctionCaller>& GetFunctionCallers();
+	UE_API const TArray<FFunctionCaller>& GetFunctionCallers() const;
 	UE_API void RemoveFunctionCallers(const TArray<FFunctionCaller*>& InFunctionCallers);
 	UE_API void ExecuteTargetFunction(FName FunctionName);
 	UE_API void ExecuteAllTargetFunctions();
diff --git a/Engine/Plugins/Enterprise/VariantManagerContent/Source/VariantManagerContent/Public/VariantSet.h b/Engine/Plugins/Enterprise/VariantManagerContent/Source/VariantManagerContent/Public/VariantSet.h
--- a/Engine/Plugins/Enterprise/VariantManagerContent/Source/VariantManagerContent/Public/VariantSet.h
+++ b/Engine/Plugins/Enterprise/VariantManagerContent/Source/VariantManagerContent/Public/VariantSet.h
@@ -41,6 +41,7 @@ public:
 	UE_API FString GetUniqueVariantName(const FString& InPrefix) const;
 
 	UE_API void AddVariants(const TArray<UVariant*>& NewVariants, int32 Index = INDEX_NONE);
+	UFUNCTION(BlueprintPure, Category="VariantSet")
 	UE_API int32 GetVariantIndex(UVariant* Var) const;
 	UE_API const TArray<UVariant*>& GetVariants() const;
 	UE_API void RemoveVariants(const TArray<UVariant*>& InVariants);
@@ -54,7 +55,13 @@ public:
 	UFUNCTION(BlueprintPure, Category="VariantSet")
 	UE_API UVariant* GetVariantByName(FString VariantName);
 
+	UFUNCTION(BlueprintPure, Category="VariantSet")
+	UE_API int32 GetVariantIndexByName(FString VariantName) const;
+
 	// Sets the thumbnail to use for this variant set. Can receive nullptr to clear it
+	UFUNCTION(BlueprintCallable, Category = "VariantSet|Thumbnail")
+	UE_API void SetThumbnailFromObject(UObject* ObjectThumbnail);
+
 	UFUNCTION(BlueprintCallable, Category = "VariantSet|Thumbnail")
 	UE_API void SetThumbnailFromTexture(UTexture2D* NewThumbnail);
 
diff --git a/Engine/Plugins/Enterprise/VariantManagerContent/Source/VariantManagerContentEditor/Private/AssetDefinition_LevelVariantSets.cpp b/Engine/Plugins/Enterprise/VariantManagerContent/Source/VariantManagerContentEditor/Private/AssetDefinition_LevelVariantSets.cpp
--- a/Engine/Plugins/Enterprise/VariantManagerContent/Source/VariantManagerContentEditor/Private/AssetDefinition_LevelVariantSets.cpp
+++ b/Engine/Plugins/Enterprise/VariantManagerContent/Source/VariantManagerContentEditor/Private/AssetDefinition_LevelVariantSets.cpp
@@ -7,6 +7,8 @@
 #include "ToolMenuSection.h"
 #include "VariantManagerContentEditorModule.h"
 #include "ToolMenus.h"
+#include "VariantManagerContentStyle.h"
+#include "Styling/SlateStyleRegistry.h"
 
 #include UE_INLINE_GENERATED_CPP_BY_NAME(AssetDefinition_LevelVariantSets)
 
@@ -49,6 +51,11 @@ FAssetOpenSupport UAssetDefinition_LevelVariantSets::GetAssetOpenSupport(const F
 	return FAssetOpenSupport(OpenSupportArgs.OpenMethod,OpenSupportArgs.OpenMethod == EAssetOpenMethod::Edit, EToolkitMode::WorldCentric); 
 }
 
+const FSlateBrush* UAssetDefinition_LevelVariantSets::GetThumbnailBrush(const FAssetData& InAssetData, const FName InClassName) const
+{
+	return FVariantManagerContentStyle::Get().GetBrush("VariantManager.Icon");
+}
+
 namespace MenuExtension_LevelVariantSets
 {
 	static FDelayedAutoRegisterHelper DelayedAutoRegister(EDelayedRegisterRunPhase::EndOfEngineInit, []{ 
diff --git a/Engine/Plugins/Enterprise/VariantManagerContent/Source/VariantManagerContentEditor/Private/AssetDefinition_LevelVariantSets.h b/Engine/Plugins/Enterprise/VariantManagerContent/Source/VariantManagerContentEditor/Private/AssetDefinition_LevelVariantSets.h
--- a/Engine/Plugins/Enterprise/VariantManagerContent/Source/VariantManagerContentEditor/Private/AssetDefinition_LevelVariantSets.h
+++ b/Engine/Plugins/Enterprise/VariantManagerContent/Source/VariantManagerContentEditor/Private/AssetDefinition_LevelVariantSets.h
@@ -11,6 +11,7 @@
 UCLASS()
 class UAssetDefinition_LevelVariantSets : public UAssetDefinitionDefault
 {
+private:
 	GENERATED_BODY()
 
 public:
@@ -23,5 +24,6 @@ public:
 	virtual bool CanImport() const override { return false; }
 	virtual FAssetSupportResponse CanLocalize(const FAssetData& InAsset) const override;
 	virtual FAssetOpenSupport GetAssetOpenSupport(const FAssetOpenSupportArgs& OpenSupportArgs) const override;
+	virtual const FSlateBrush* GetThumbnailBrush(const FAssetData& InAssetData, const FName InClassName) const override;
 	// UAssetDefinition End
 };
diff --git a/Engine/Plugins/Enterprise/VariantManagerContent/Source/VariantManagerContentEditor/Private/VariantManagerContentEditorModule.cpp b/Engine/Plugins/Enterprise/VariantManagerContent/Source/VariantManagerContentEditor/Private/VariantManagerContentEditorModule.cpp
--- a/Engine/Plugins/Enterprise/VariantManagerContent/Source/VariantManagerContentEditor/Private/VariantManagerContentEditorModule.cpp
+++ b/Engine/Plugins/Enterprise/VariantManagerContent/Source/VariantManagerContentEditor/Private/VariantManagerContentEditorModule.cpp
@@ -15,6 +15,7 @@
 #include "ObjectTools.h"
 #include "PackageTools.h"
 #include "PropertyEditorModule.h"
+#include "VariantManagerContentStyle.h"
 
 #define LOCTEXT_NAMESPACE "VariantManagerContentEditorModule"
 
@@ -25,6 +26,8 @@ class FVariantManagerContentEditorModule : public IVariantManagerContentEditorMo
 public:
 	virtual void StartupModule() override
 	{
+		FVariantManagerContentStyle::Initialize();
+
 		FPropertyEditorModule& PropertyModule = FModuleManager::LoadModuleChecked<FPropertyEditorModule>(TEXT("PropertyEditor"));
 		PropertyModule.RegisterCustomClassLayout(TEXT("LevelVariantSetsActor"), FOnGetDetailCustomizationInstance::CreateStatic(&FLevelVariantSetsActorCustomization::MakeInstance));
 		PropertyModule.RegisterCustomClassLayout(TEXT("SwitchActor"), FOnGetDetailCustomizationInstance::CreateStatic(&FSwitchActorCustomization::MakeInstance));
@@ -32,9 +35,13 @@ public:
 
 	virtual void ShutdownModule() override
 	{
-        FPropertyEditorModule& PropertyModule = FModuleManager::LoadModuleChecked< FPropertyEditorModule >(TEXT("PropertyEditor"));
-		PropertyModule.UnregisterCustomClassLayout(TEXT("LevelVariantSetsActor"));
-		PropertyModule.UnregisterCustomClassLayout(TEXT("SwitchActor"));
+        if (FPropertyEditorModule* PropertyModule = FModuleManager::LoadModulePtr< FPropertyEditorModule >(TEXT("PropertyEditor")))
+        {
+        	PropertyModule->UnregisterCustomClassLayout(TEXT("LevelVariantSetsActor"));
+			PropertyModule->UnregisterCustomClassLayout(TEXT("SwitchActor"));
+		}
+
+		FVariantManagerContentStyle::Shutdown();
 	}
 
 	virtual void RegisterOnLevelVariantSetsDelegate(FOnLevelVariantSetsEditor Delegate) override
diff --git a/Engine/Plugins/Enterprise/VariantManagerContent/Source/VariantManagerContentEditor/Private/VariantManagerContentStyle.cpp b/Engine/Plugins/Enterprise/VariantManagerContent/Source/VariantManagerContentEditor/Private/VariantManagerContentStyle.cpp
new file mode 100644
--- /dev/null
+++ b/Engine/Plugins/Enterprise/VariantManagerContent/Source/VariantManagerContentEditor/Private/VariantManagerContentStyle.cpp
@@ -0,0 +1,44 @@
+// Copyright Epic Games, Inc. All Rights Reserved.
+
+
+#include "VariantManagerContentStyle.h"
+
+#include "Interfaces/IPluginManager.h"
+#include "Styling/SlateStyleMacros.h"
+
+
+FVariantManagerContentStyle::FVariantManagerContentStyle() : FSlateStyleSet("VariantManagerContentStyle")
+{
+	if (TSharedPtr<IPlugin> Plugin = IPluginManager::Get().FindPlugin(TEXT("VariantManagerContent")))
+	{
+		SetContentRoot(Plugin->GetContentDir());
+	}
+
+	// Const icon sizes
+	const FVector2D Icon8x8(8.0f, 8.0f);
+	const FVector2D Icon16x16(16.0f, 16.0f);
+	const FVector2D Icon20x20(20.0f, 20.0f);
+	const FVector2D Icon64x64(64.0f, 64.0f);
+
+	Set( "VariantManager.Icon", new IMAGE_BRUSH_SVG("VariantManager_64", Icon64x64) );
+
+	FSlateStyleRegistry::RegisterSlateStyle(*this);
+
+}
+
+void FVariantManagerContentStyle::Initialize()
+{
+	FSlateStyleRegistry::UnRegisterSlateStyle(Get());
+
+	FSlateStyleRegistry::RegisterSlateStyle(Get());
+}
+
+void FVariantManagerContentStyle::Shutdown()
+{
+	FSlateStyleRegistry::UnRegisterSlateStyle(Get());
+}
+
+const FName FVariantManagerContentStyle::GetAppStyleSetName()
+{
+	return FName("VariantManagerContentStyle");
+}
diff --git a/Engine/Plugins/Enterprise/VariantManagerContent/Source/VariantManagerContentEditor/Private/VariantManagerContentStyle.h b/Engine/Plugins/Enterprise/VariantManagerContent/Source/VariantManagerContentEditor/Private/VariantManagerContentStyle.h
new file mode 100644
--- /dev/null
+++ b/Engine/Plugins/Enterprise/VariantManagerContent/Source/VariantManagerContentEditor/Private/VariantManagerContentStyle.h
@@ -0,0 +1,24 @@
+// Copyright Epic Games, Inc. All Rights Reserved.
+
+#pragma once
+
+#include "Styling/SlateStyle.h"
+#include "Styling/SlateStyleRegistry.h"
+
+struct FVariantManagerContentStyle : public FSlateStyleSet
+{
+public:
+	FVariantManagerContentStyle();
+
+	static void Initialize();
+
+	static void Shutdown();
+
+	static const FName GetAppStyleSetName();
+
+	static FVariantManagerContentStyle& Get()
+	{
+		static FVariantManagerContentStyle Inst;
+		return Inst;
+	}
+};
\ No newline at end of file
diff --git a/Engine/Plugins/Enterprise/VariantManagerContent/Source/VariantManagerContentEditor/VariantManagerContentEditor.Build.cs b/Engine/Plugins/Enterprise/VariantManagerContent/Source/VariantManagerContentEditor/VariantManagerContentEditor.Build.cs
--- a/Engine/Plugins/Enterprise/VariantManagerContent/Source/VariantManagerContentEditor/VariantManagerContentEditor.Build.cs
+++ b/Engine/Plugins/Enterprise/VariantManagerContent/Source/VariantManagerContentEditor/VariantManagerContentEditor.Build.cs
@@ -25,6 +25,7 @@ namespace UnrealBuildTool.Rules
 				new string[]
 				{
 					"InputCore", // For ListView keyboard control
+					"Projects",
 					"Slate",
 					"SlateCore",
 					"VariantManagerContent",
